[Feature] Archetype Fixup Tool

diff --git a/Engine/Plugins/Editor/ArchetypeFixupTool/ArchetypeFixupTool.uplugin b/Engine/Plugins/Editor/ArchetypeFixupTool/ArchetypeFixupTool.uplugin
new file mode 100644
--- /dev/null
+++ b/Engine/Plugins/Editor/ArchetypeFixupTool/ArchetypeFixupTool.uplugin
@@ -0,0 +1,32 @@
+{
+	"FileVersion" : 3,
+	"Version" : 1,
+	"VersionName" : "1.0",
+	"FriendlyName" : "Archetype Fixup Tool",
+	"Description" : "Editor tool for redirecting loose archetype properties",
+	"Category" : "Editor",
+	"CreatedBy" : "Epic Games, Inc.",
+	"CreatedByURL" : "https://epicgames.com",
+	"DocsURL" : "",
+	"MarketplaceURL" : "",
+	"SupportURL" : "",
+	"EnabledByDefault" : true,
+	"CanContainContent" : false,
+	"IsBetaVersion" : true,
+	"Installed" : false,
+	"Plugins": 
+	[
+		{
+			"Name": "StructUtils",
+			"Enabled": true
+		}
+	],
+	"Modules" :
+	[
+		{
+			"Name" : "ArchetypeFixupTool",
+			"Type" : "Editor",
+			"LoadingPhase" : "Default"
+		}
+	]
+}
\ No newline at end of file
diff --git a/Engine/Plugins/Editor/ArchetypeFixupTool/Source/ArchetypeFixupTool.Build.cs b/Engine/Plugins/Editor/ArchetypeFixupTool/Source/ArchetypeFixupTool.Build.cs
new file mode 100644
--- /dev/null
+++ b/Engine/Plugins/Editor/ArchetypeFixupTool/Source/ArchetypeFixupTool.Build.cs
@@ -0,0 +1,29 @@
+// Copyright Epic Games, Inc. All Rights Reserved.
+
+namespace UnrealBuildTool.Rules
+{
+	public class ArchetypeFixupTool : ModuleRules
+	{
+		public ArchetypeFixupTool(ReadOnlyTargetRules Target) : base(Target)
+		{
+			PublicDependencyModuleNames.AddRange(
+				new string[] {
+					"Core",
+					"Engine",
+					"CoreUObject",
+					"StructUtils",
+				}
+			);
+			
+			PrivateDependencyModuleNames.AddRange(
+				new string[] {
+					"SlateCore",
+					"Slate",
+					"PropertyEditor", 
+					"WorkspaceMenuStructure", 
+					"UnrealEd",
+				}
+			);
+		}
+	}
+}
\ No newline at end of file
diff --git a/Engine/Plugins/Editor/ArchetypeFixupTool/Source/Private/ArchetypeFixupDetailCustomization.cpp b/Engine/Plugins/Editor/ArchetypeFixupTool/Source/Private/ArchetypeFixupDetailCustomization.cpp
new file mode 100644
--- /dev/null
+++ b/Engine/Plugins/Editor/ArchetypeFixupTool/Source/Private/ArchetypeFixupDetailCustomization.cpp
@@ -0,0 +1,464 @@
+// Copyright Epic Games, Inc. All Rights Reserved.
+
+
+#include "ArchetypeFixupDetailCustomization.h"
+
+#include "DetailLayoutBuilder.h"
+#include "DetailWidgetRow.h"
+#include "IDetailChildrenBuilder.h"
+#include "ArchetypeFixupPanel.h"
+#include "Widgets/Layout/SWidgetSwitcher.h"
+#include "Widgets/Input/SComboButton.h"
+#include "Framework/MultiBox/MultiBoxBuilder.h"
+
+#define LOCTEXT_NAMESPACE "ArchetypeFixupDetails"
+
+
+TSet<FPropertyPath> FArchetypeFixupDetailNodeBuilder::GetRedirectOptions(const UStruct* Struct, void* Value) const
+{
+	TSet<FPropertyPath> Result;
+	GetRedirectOptions(Struct, Value, {}, Result);
+	return Result;
+}
+
+void FArchetypeFixupDetailNodeBuilder::GetRedirectOptions(const UStruct* Struct, void* Value, const FPropertyPath& Path, TSet<FPropertyPath>& OutPaths) const
+{
+	for (FProperty* SubProperty : TFieldRange<FProperty>(Struct))
+    {
+		if (SubProperty->ArrayDim == 1)
+		{
+			const FPropertyPath SubPath = Path.ExtendPath(FPropertyInfo(SubProperty)).Get();
+			GetRedirectOptions(SubProperty, SubProperty->ContainerPtrToValuePtr<void>(Value), SubPath, OutPaths);
+		}
+		else
+		{
+			for (int32 ArrayIndex = 0; ArrayIndex < SubProperty->ArrayDim; ++ArrayIndex)
+			{
+    			const FPropertyPath SubPath = Path.ExtendPath(FPropertyInfo(SubProperty, ArrayIndex)).Get();
+    			GetRedirectOptions(SubProperty, SubProperty->ContainerPtrToValuePtr<void>(Value, ArrayIndex), SubPath, OutPaths);
+			}
+		}
+    }
+}
+
+void FArchetypeFixupDetailNodeBuilder::GetRedirectOptions(const FProperty* Property, void* Value, const FPropertyPath& Path, TSet<FPropertyPath>& OutPaths) const
+{
+	if (Property->GetBoolMetaData(TEXT("isLoose")))
+	{
+		// don't include loose properties as options
+		return;
+	}
+
+	if (!DiffPanel.Pin()->RedirectedPropertyTree->Find(Path))
+	{
+		OutPaths.Add(Path);
+	}
+	
+	if (const FStructProperty* AsStructProperty = CastField<FStructProperty>(Property))
+	{
+		GetRedirectOptions(AsStructProperty->Struct, Value, Path, OutPaths);
+	}
+	else if (const FArrayProperty* AsArrayProperty = CastField<FArrayProperty>(Property))
+	{
+		FScriptArrayHelper Array(AsArrayProperty, Value);
+		for (int32 ArrayIndex = 0; ArrayIndex < Array.Num(); ++ArrayIndex)
+		{
+			const FPropertyPath SubPath = Path.ExtendPath(FPropertyInfo(AsArrayProperty->Inner, ArrayIndex)).Get();
+			GetRedirectOptions(AsArrayProperty->Inner, Array.GetElementPtr(ArrayIndex), SubPath, OutPaths);
+		}
+	}
+	else if (const FSetProperty* AsSetProperty = CastField<FSetProperty>(Property))
+	{
+		FScriptSetHelper Set(AsSetProperty, Value);
+		for (FScriptSetHelper::FIterator Iterator = Set.CreateIterator(); Iterator; ++Iterator)
+		{
+			const FPropertyPath SubPath = Path.ExtendPath(FPropertyInfo(AsSetProperty->ElementProp, Iterator.GetLogicalIndex())).Get();
+			GetRedirectOptions(AsSetProperty->ElementProp, Set.GetElementPtr(Iterator), SubPath, OutPaths);
+		}
+	}
+	else if (const FMapProperty* AsMapProperty = CastField<FMapProperty>(Property))
+	{
+		FScriptMapHelper Map(AsMapProperty, Value);
+		for (FScriptMapHelper::FIterator Iterator = Map.CreateIterator(); Iterator; ++Iterator)
+		{
+			const FPropertyPath SubPath = Path.ExtendPath(FPropertyInfo(AsMapProperty->ValueProp, Iterator.GetLogicalIndex())).Get();
+			GetRedirectOptions(AsMapProperty->ValueProp, Map.GetValuePtr(Iterator), SubPath, OutPaths);
+		}
+	}
+}
+
+/////////////////////////////////////////////////////////////////
+// FArchetypeFixupDetailCustomization
+/////////////////////////////////////////////////////////////////
+
+FArchetypeFixupDetailCustomization::FArchetypeFixupDetailCustomization(const TSharedRef<FArchetypeFixupPanel>& InDiffPanel)
+	: DiffPanel(InDiffPanel)
+{
+	
+}
+
+FArchetypeFixupDetailCustomization::~FArchetypeFixupDetailCustomization()
+{
+}
+
+void FArchetypeFixupDetailCustomization::CustomizeDetails(IDetailLayoutBuilder& DetailBuilder)
+{
+	const TSharedPtr<FArchetypeFixupPanel> Panel = DiffPanel.Pin();
+	if (!Panel)
+	{
+		return;
+	}
+	
+	TArray<TWeakObjectPtr<UObject>> Objects;
+	DetailBuilder.GetObjectsBeingCustomized(Objects);
+	if (Objects.IsEmpty())
+	{
+		return;
+	}
+	TArray<FName> CategoryNames;
+	DetailBuilder.GetCategoryNames(CategoryNames);
+	for (const FName CategoryName : CategoryNames)
+	{
+		IDetailCategoryBuilder& Category = DetailBuilder.EditCategory(CategoryName);
+		TArray<TSharedRef<IPropertyHandle>> NonAdvancedProperties;
+		Category.GetDefaultProperties(NonAdvancedProperties, true, false);
+		for (const TSharedRef<IPropertyHandle>& Handle : NonAdvancedProperties)
+		{
+			DetailBuilder.HideProperty(Handle);
+			Category.AddCustomBuilder(MakeShared<FArchetypeFixupDetailNodeBuilder>(Panel.ToSharedRef(), Handle), false);
+		}
+		TArray<TSharedRef<IPropertyHandle>> AdvancedProperties;
+		Category.GetDefaultProperties(AdvancedProperties, false, true);
+		for (const TSharedRef<IPropertyHandle>& Handle : AdvancedProperties)
+		{
+			DetailBuilder.HideProperty(Handle);
+			Category.AddCustomBuilder(MakeShared<FArchetypeFixupDetailNodeBuilder>(Panel.ToSharedRef(), Handle), true);
+		}
+	}
+}
+
+/////////////////////////////////////////////////////////////////
+// FHideLoosePropertiesCustomization
+/////////////////////////////////////////////////////////////////
+
+FHideLoosePropertiesCustomization::~FHideLoosePropertiesCustomization()
+{
+}
+
+void FHideLoosePropertiesCustomization::CustomizeDetails(IDetailLayoutBuilder& DetailBuilder)
+{
+	TArray<TWeakObjectPtr<UObject>> Objects;
+	DetailBuilder.GetObjectsBeingCustomized(Objects);
+	if (Objects.IsEmpty())
+	{
+		return;
+	}
+	TArray<FName> CategoryNames;
+	DetailBuilder.GetCategoryNames(CategoryNames);
+	for (const FName CategoryName : CategoryNames)
+	{
+		IDetailCategoryBuilder& Category = DetailBuilder.EditCategory(CategoryName);
+		TArray<TSharedRef<IPropertyHandle>> Handles;
+		Category.GetDefaultProperties(Handles, true, true);
+		for (const TSharedRef<IPropertyHandle>& Handle : Handles)
+		{
+			CustomizeHandle(Handle, DetailBuilder);
+		}
+	}
+}
+
+void FHideLoosePropertiesCustomization::CustomizeHandle(const TSharedRef<IPropertyHandle>& Handle, IDetailLayoutBuilder& DetailBuilder)
+{
+	if (Handle->GetBoolMetaData(TEXT("isLoose")))
+	{
+		DetailBuilder.HideProperty(Handle);
+	}
+	else
+	{
+		// recurse into children
+		uint32 NumChildren = 0;
+		Handle->GetNumChildren(NumChildren);
+		for (uint32 ChildIndex = 0; ChildIndex < NumChildren; ++ChildIndex)
+		{
+			CustomizeHandle(Handle->GetChildHandle(ChildIndex).ToSharedRef(), DetailBuilder);
+		}
+	}
+}
+
+
+/////////////////////////////////////////////////////////////////
+// FArchetypeFixupDetailNodeBuilder
+/////////////////////////////////////////////////////////////////
+
+FArchetypeFixupDetailNodeBuilder::FArchetypeFixupDetailNodeBuilder(const TSharedRef<FArchetypeFixupPanel>& InDiffPanel, const TSharedRef<IPropertyHandle>& InPropertyHandle)
+	: DiffPanel(InDiffPanel)
+	, PropertyHandle(InPropertyHandle)
+{}
+
+void FArchetypeFixupDetailNodeBuilder::GenerateHeaderRowContent(FDetailWidgetRow& NodeRow)
+{	
+	const TSharedPtr<FArchetypeFixupPanel> Panel = DiffPanel.Pin();
+	if (!Panel)
+	{
+		return;
+	}
+
+	// skip properties that are hidden from this view mode
+	if (IsHidden())
+	{
+		return;
+	}
+	
+	TSharedPtr<SWidget> InnerNameContent = PropertyHandle->CreatePropertyNameWidget();
+	if (PropertyHandle->GetKeyHandle())
+	{
+		// TODO: @jordan.hoffmann this solution for keys only works with basic types. figure out a more robust solution.
+		InnerNameContent = PropertyHandle->GetKeyHandle()->CreatePropertyValueWidget();
+		InnerNameContent->SetEnabled(false);
+	}
+
+	const TSharedRef<SWidget> NameContent = SNew(SWidgetSwitcher)
+		.WidgetIndex(this, &FArchetypeFixupDetailNodeBuilder::GetNameWidgetIndex)
+		+SWidgetSwitcher::Slot()
+		[
+			InnerNameContent.ToSharedRef()
+		]
+		+SWidgetSwitcher::Slot()
+		[
+			SNew(SHorizontalBox)
+			+SHorizontalBox::Slot()
+			.VAlign(EVerticalAlignment::VAlign_Center)
+			.AutoWidth()
+			[
+				SNew(STextBlock)
+				.Visibility(this, &FArchetypeFixupDetailNodeBuilder::DeletionSymbolVisibility)
+				.Text(LOCTEXT("MarkedForDeletion", "❌"))
+			]
+			+SHorizontalBox::Slot()
+			.VAlign(EVerticalAlignment::VAlign_Center)
+			.AutoWidth()
+			[
+				SNew(SComboButton)
+				.Visibility(EVisibility::Visible)
+				.OnGetMenuContent_Raw(this, &FArchetypeFixupDetailNodeBuilder::GeneratePropertyRedirectMenu)
+				.ButtonContent()
+				[
+					InnerNameContent.ToSharedRef()
+				]
+			]
+		];
+	
+	const TSharedRef<SWidget> ValueContent = PropertyHandle->CreatePropertyValueWidget();
+	ValueContent->SetEnabled(!Panel->HasViewFlag(FArchetypeFixupPanel::EViewFlags::ReadonlyValues));
+	ValueContent->SetVisibility(TAttribute<EVisibility>(this, &FArchetypeFixupDetailNodeBuilder::ValueContentVisibility));
+	
+	NodeRow
+	.NameContent()
+	[
+		NameContent
+	]
+	.ValueContent()
+	[
+		ValueContent
+	]
+	.PropertyHandleList({PropertyHandle});
+}
+
+void FArchetypeFixupDetailNodeBuilder::GenerateChildContent(IDetailChildrenBuilder& ChildrenBuilder)
+{
+	const TSharedPtr<FArchetypeFixupPanel> Panel = DiffPanel.Pin();
+	if (!Panel)
+	{
+		return;
+	}
+
+	// skip properties that are hidden from this view mode
+	if (IsHidden())
+	{
+		return;
+	}
+	
+	uint32 ChildCount;
+	PropertyHandle->GetNumChildren(ChildCount);
+	for (uint32 I = 0; I < ChildCount; ++I)
+	{
+		const TSharedRef<IPropertyHandle> ChildHandle = PropertyHandle->GetChildHandle(I).ToSharedRef();
+		ChildrenBuilder.AddCustomBuilder(MakeShared<FArchetypeFixupDetailNodeBuilder>(Panel.ToSharedRef(), ChildHandle));
+	}
+}
+
+FName FArchetypeFixupDetailNodeBuilder::GetName() const
+{
+	return PropertyHandle->GetProperty()->GetFName();
+}
+
+TSharedPtr<IPropertyHandle> FArchetypeFixupDetailNodeBuilder::GetPropertyHandle() const
+{
+	return PropertyHandle;
+}
+
+int32 FArchetypeFixupDetailNodeBuilder::GetNameWidgetIndex() const
+{
+	if (const TSharedPtr<FArchetypeFixupPanel> Panel = DiffPanel.Pin())
+	{
+		if (Panel->HasViewFlag(FArchetypeFixupPanel::EViewFlags::AllowRemapLooseProperties))
+		{
+			if (Panel->OriginalPaths.Contains(*PropertyHandle->CreateFPropertyPath()))
+			{
+				return DisplayRedirectMenu;
+			}
+			if (const FProperty* Property = PropertyHandle->GetProperty())
+			{
+				if (Property->GetBoolMetaData(TEXT("isLoose")))
+				{
+					return DisplayRedirectMenu;
+				}
+			}
+		}
+	}
+	return DisplayRegularName;
+}
+
+TSharedRef<SWidget> FArchetypeFixupDetailNodeBuilder::GeneratePropertyRedirectMenu() const
+{
+	FMenuBuilder MenuBuilder(true, nullptr);
+	
+	const TSharedPtr<FArchetypeFixupPanel> Panel = DiffPanel.Pin();
+	if (!Panel)
+	{
+		return MenuBuilder.MakeWidget();
+	}
+	const FPropertyPath Path = *PropertyHandle->CreateFPropertyPath();
+	const FPropertyPath& OriginalPath = Panel->GetOriginalPath(Path);
+
+	MenuBuilder.BeginSection(NAME_None, LOCTEXT("ResetRedirect", "Reset"));
+	if (OriginalPath != Path || Panel->MarkedForDelete.Contains(OriginalPath))
+	{
+		FText OriginalPathText = FText::FromString(OriginalPath.ToString());
+		FText Tooltip = FText::Format(LOCTEXT("ResetTooltip", "Reset back to {0}"), OriginalPathText);
+		MenuBuilder.AddMenuEntry(OriginalPathText, Tooltip, FSlateIcon()
+						, FUIAction(FExecuteAction::CreateSP(Panel.Get(), &FArchetypeFixupPanel::OnRedirectProperty, Path, OriginalPath))
+						, NAME_None
+						, EUserInterfaceActionType::RadioButton);
+	}
+	MenuBuilder.EndSection();
+
+	UObject* FirstArchetype = Panel->Instances[0];
+	TSet<FPropertyPath> RedirectOptions = GetRedirectOptions(FirstArchetype->GetClass(), FirstArchetype);
+
+	MenuBuilder.BeginSection(NAME_None, LOCTEXT("MoveToMatchingRow", "Move"));
+	{
+		for (const FPropertyPath& Option : RedirectOptions)
+		{
+			FName ThisPropName = PropertyHandle->GetProperty()->GetFName();
+			FName OptionPropName = Option.GetLeafMostProperty().Property->GetFName();
+			if (ThisPropName == OptionPropName)
+			{
+				if (Option.GetLeafMostProperty().Property->SameType(PropertyHandle->GetProperty()))
+				{
+					FText DisplayName = FText::FromString(Option.ToString());
+					FText Tooltip = FText::Format(LOCTEXT("ChangeToSameTypeTooltip", "Move property to '{0}'"), DisplayName);
+					MenuBuilder.AddMenuEntry(DisplayName, Tooltip, FSlateIcon()
+					, FUIAction(FExecuteAction::CreateSP(Panel.Get(), &FArchetypeFixupPanel::OnRedirectProperty, Path, Option))
+					, NAME_None
+					, EUserInterfaceActionType::RadioButton);
+				}
+			}
+		}
+	}
+	MenuBuilder.EndSection();
+
+	MenuBuilder.BeginSection(NAME_None, LOCTEXT("ChangeToSameType", "Rename"));
+	{
+		for (const FPropertyPath& Option : RedirectOptions)
+		{
+			FName ThisPropName = PropertyHandle->GetProperty()->GetFName();
+			FName OptionPropName = Option.GetLeafMostProperty().Property->GetFName();
+			if (ThisPropName == OptionPropName)
+			{
+				continue; // handled in the "move" category
+			}
+			if (Option.GetLeafMostProperty().Property->SameType(PropertyHandle->GetProperty()))
+			{
+				FText DisplayName = FText::FromString(Option.ToString());
+				FText Tooltip = FText::Format(LOCTEXT("ChangeToSameTypeTooltip", "Rename property to '{0}'"), DisplayName);
+				MenuBuilder.AddMenuEntry(DisplayName, Tooltip, FSlateIcon()
+				, FUIAction(FExecuteAction::CreateSP(Panel.Get(), &FArchetypeFixupPanel::OnRedirectProperty, Path, Option))
+				, NAME_None
+				, EUserInterfaceActionType::RadioButton);
+			}
+		}
+	}
+	MenuBuilder.EndSection();
+
+	MenuBuilder.BeginSection(NAME_None, LOCTEXT("ChangeToDifferingTypeConversion", "Convert Type (potential data loss)"));
+	{
+		// TODO: Add type conversions?
+	}
+	MenuBuilder.EndSection();
+	
+	MenuBuilder.BeginSection(NAME_None, LOCTEXT("Delete", "Delete"));
+	{
+		if (!Panel->MarkedForDelete.Contains(OriginalPath)) // check that it wasn't already marked as deleted
+		{
+			FText DisplayName = LOCTEXT("MarkForDeletion", "Mark For Deletion");
+			FText Tooltip = LOCTEXT("MarkForDeletion", "Mark this property for deletion");
+			MenuBuilder.AddMenuEntry(DisplayName, Tooltip, FSlateIcon()
+						, FUIAction(FExecuteAction::CreateSP(Panel.Get(), &FArchetypeFixupPanel::OnMarkForDelete, Path))
+						, NAME_None
+						, EUserInterfaceActionType::RadioButton);
+		}
+	}
+	MenuBuilder.EndSection();
+	
+	return MenuBuilder.MakeWidget();
+}
+
+EVisibility FArchetypeFixupDetailNodeBuilder::DeletionSymbolVisibility() const
+{
+	if (const TSharedPtr<FArchetypeFixupPanel> Panel = DiffPanel.Pin())
+	{
+		const FPropertyPath Path = *PropertyHandle->CreateFPropertyPath();
+		return Panel->MarkedForDelete.Contains(Path) ? EVisibility::Visible :  EVisibility::Collapsed;
+	}
+	return EVisibility::Collapsed;
+}
+
+EVisibility FArchetypeFixupDetailNodeBuilder::ValueContentVisibility() const
+{
+	if (const TSharedPtr<FArchetypeFixupPanel> Panel = DiffPanel.Pin())
+	{
+		const FPropertyPath Path = *PropertyHandle->CreateFPropertyPath();
+		return Panel->MarkedForDelete.Contains(Path) ? EVisibility::Collapsed :  EVisibility::Visible;
+	}
+	return EVisibility::Visible;
+}
+
+bool FArchetypeFixupDetailNodeBuilder::IsHidden() const
+{
+	const TSharedPtr<FArchetypeFixupPanel> Panel = DiffPanel.Pin();
+	if (!Panel)
+	{
+		return true;
+	}
+
+	if (const FProperty* Property = PropertyHandle->GetProperty())
+	{
+		if (Panel->HasViewFlag(FArchetypeFixupPanel::EViewFlags::HideLooseProperties) &&
+			Property->GetBoolMetaData(TEXT("isLoose")))
+		{
+			return true;
+		}
+
+		if (Panel->HasViewFlag(FArchetypeFixupPanel::EViewFlags::IncludeOnlySetBySerialization))
+		{
+			if (!Panel->IsInRedirectedPropertyTree(*PropertyHandle->CreateFPropertyPath()))
+			{
+				return true;
+			}
+		}
+	}
+	return false;
+}
+
+#undef LOCTEXT_NAMESPACE
diff --git a/Engine/Plugins/Editor/ArchetypeFixupTool/Source/Private/ArchetypeFixupDetailCustomization.h b/Engine/Plugins/Editor/ArchetypeFixupTool/Source/Private/ArchetypeFixupDetailCustomization.h
new file mode 100644
--- /dev/null
+++ b/Engine/Plugins/Editor/ArchetypeFixupTool/Source/Private/ArchetypeFixupDetailCustomization.h
@@ -0,0 +1,70 @@
+// Copyright Epic Games, Inc. All Rights Reserved.
+
+#pragma once
+
+#include "CoreMinimal.h"
+#include "IDetailCustomization.h"
+#include "IDetailCustomNodeBuilder.h"
+#include "Containers/Set.h"
+#include "PropertyPath.h"
+#include "Widgets/SWidget.h"
+
+class IDetailCategoryBuilder;
+class FArchetypeFixupPanel;
+
+/**
+ * 
+ */
+class ARCHETYPEFIXUPTOOL_API FArchetypeFixupDetailCustomization : public IDetailCustomization, public TSharedFromThis<FArchetypeFixupDetailCustomization>
+{
+public:
+	FArchetypeFixupDetailCustomization(const TSharedRef<FArchetypeFixupPanel>& DiffPanel);
+	virtual ~FArchetypeFixupDetailCustomization() override;
+	virtual void CustomizeDetails(IDetailLayoutBuilder& DetailBuilder) override;
+
+private:
+	TWeakPtr<FArchetypeFixupPanel> DiffPanel;
+};
+
+class ARCHETYPEFIXUPTOOL_API FHideLoosePropertiesCustomization : public IDetailCustomization, public TSharedFromThis<FArchetypeFixupDetailCustomization>
+{
+public:
+	FHideLoosePropertiesCustomization() = default;
+	virtual ~FHideLoosePropertiesCustomization() override;
+	virtual void CustomizeDetails(IDetailLayoutBuilder& DetailBuilder) override;
+
+	static void CustomizeHandle(const TSharedRef<IPropertyHandle>& Handle, IDetailLayoutBuilder& DetailBuilder);
+
+private:
+	TWeakPtr<FArchetypeFixupPanel> DiffPanel;
+};
+
+class ARCHETYPEFIXUPTOOL_API FArchetypeFixupDetailNodeBuilder : public IDetailCustomNodeBuilder, public TSharedFromThis<FArchetypeFixupDetailNodeBuilder>
+{
+public:
+	FArchetypeFixupDetailNodeBuilder(const TSharedRef<FArchetypeFixupPanel>& DiffPanel, const TSharedRef<IPropertyHandle>& PropertyHandle);
+	virtual void GenerateHeaderRowContent(FDetailWidgetRow& NodeRow) override;
+	virtual void GenerateChildContent(IDetailChildrenBuilder& ChildrenBuilder) override;
+	virtual FName GetName() const override;
+	virtual TSharedPtr<IPropertyHandle> GetPropertyHandle() const override;
+
+private:
+	int32 GetNameWidgetIndex() const;
+	TSharedRef<SWidget> GeneratePropertyRedirectMenu() const;
+	EVisibility DeletionSymbolVisibility() const;
+	EVisibility ValueContentVisibility() const;
+	bool IsHidden() const;
+	
+	TSet<FPropertyPath> GetRedirectOptions(const UStruct* Struct, void* Value) const;
+	void GetRedirectOptions(const UStruct* Struct, void* Value, const FPropertyPath& Path, TSet<FPropertyPath>& OutPaths) const;
+	void GetRedirectOptions(const FProperty* Property, void* Value, const FPropertyPath& Path, TSet<FPropertyPath>& OutPaths) const;
+	
+	TWeakPtr<FArchetypeFixupPanel> DiffPanel;
+	TSharedRef<IPropertyHandle> PropertyHandle;
+
+	enum ENameWidgetIndex : uint8
+	{
+		DisplayRegularName = 0,
+		DisplayRedirectMenu = 1
+	};
+};
diff --git a/Engine/Plugins/Editor/ArchetypeFixupTool/Source/Private/ArchetypeFixupPanel.cpp b/Engine/Plugins/Editor/ArchetypeFixupTool/Source/Private/ArchetypeFixupPanel.cpp
new file mode 100644
--- /dev/null
+++ b/Engine/Plugins/Editor/ArchetypeFixupTool/Source/Private/ArchetypeFixupPanel.cpp
@@ -0,0 +1,585 @@
+// Copyright Epic Games, Inc. All Rights Reserved.
+
+
+#include "ArchetypeFixupPanel.h"
+
+#include "AsyncDetailViewDiff.h"
+#include "DetailTreeNode.h"
+#include "Widgets/Layout/LinkableScrollBar.h"
+#include "ArchetypeFixupDetailCustomization.h"
+#include "UObject/ArchetypeUtils.h"
+#include "Modules/ModuleManager.h"
+#include "Editor.h"
+
+#define LOCTEXT_NAMESPACE "ArchetypeFixupPanel"
+
+FRedirectedPropertyNode::FRedirectedPropertyNode(const FRedirectedPropertyNode& Other)
+	: PropertyName(Other.PropertyName)
+	, Type(Other.Type)
+	, ArrayIndex(Other.ArrayIndex)
+{
+	// deep copy tree
+	for (const TSharedPtr<FRedirectedPropertyNode>& Child : Other.Children)
+	{
+		Children.Add(MakeShared<FRedirectedPropertyNode>(*Child));
+	}
+}
+
+FRedirectedPropertyNode::FRedirectedPropertyNode(const FPropertyInfo& InInfo, const TWeakPtr<FRedirectedPropertyNode>& InParent)
+	: PropertyName(InInfo.Property->GetFName())
+	, Type(InInfo.Property->GetID())
+	, ArrayIndex(InInfo.ArrayIndex)
+	, Parent(InParent)
+{
+}
+
+FRedirectedPropertyNode::FRedirectedPropertyNode(FName InPropertyName, FName InType, int32 InArrayIndex, const TWeakPtr<FRedirectedPropertyNode>& InParent)
+	: PropertyName(InPropertyName)
+	, Type(InType)
+	, ArrayIndex(InArrayIndex)
+	, Parent(InParent)
+{
+}
+
+TSharedPtr<FRedirectedPropertyNode> FRedirectedPropertyNode::FindOrAdd(const FPropertyPath& Path, int32 PathIndex)
+{
+	check(PathIndex <= Path.GetNumProperties());
+	if (PathIndex == Path.GetNumProperties())
+	{
+		return SharedThis(this);
+	}
+	
+	const FPropertyInfo& ChildInfo = Path.GetPropertyInfo(PathIndex);
+	const TSharedPtr<FRedirectedPropertyNode> Child = FindOrAdd(ChildInfo);
+	return Child->FindOrAdd(Path, PathIndex + 1);
+}
+
+TSharedPtr<FRedirectedPropertyNode> FRedirectedPropertyNode::FindOrAdd(const FPropertyInfo& ChildInfo)
+{
+	TSharedPtr<FRedirectedPropertyNode> Child = Find(ChildInfo);
+	if (!Child)
+	{
+		Child = MakeShared<FRedirectedPropertyNode>(ChildInfo, SharedThis(this));
+		Children.Add(Child);
+	}
+	return Child;
+}
+
+TSharedPtr<FRedirectedPropertyNode> FRedirectedPropertyNode::FindOrAdd(FName ChildPropertyName, FName ChildType, int32 ChildArrayIndex)
+{
+	TSharedPtr<FRedirectedPropertyNode> Child = Find(ChildPropertyName, ChildType, ChildArrayIndex);
+	if (!Child)
+	{
+		Child = MakeShared<FRedirectedPropertyNode>(ChildPropertyName, ChildType, ChildArrayIndex, SharedThis(this));
+		Children.Add(Child);
+	}
+	return Child;
+}
+
+bool FRedirectedPropertyNode::Remove(const FPropertyPath& Path, int32 PathIndex)
+{
+	if (TSharedPtr<FRedirectedPropertyNode> NodeToRemove = Find(Path, PathIndex))
+	{
+		do
+		{
+			const TSharedPtr<FRedirectedPropertyNode> ParentNode = NodeToRemove->Parent.Pin();
+			if (ParentNode)
+			{
+				ParentNode->Remove(NodeToRemove->PropertyName, NodeToRemove->Type, NodeToRemove->ArrayIndex);
+			}
+			NodeToRemove = ParentNode;
+		} while (NodeToRemove && NodeToRemove->Children.IsEmpty());
+		return true;
+	}
+	return false;
+}
+
+bool FRedirectedPropertyNode::Remove(const FPropertyInfo& ChildInfo)
+{
+	const int32 Index = FindIndex(ChildInfo);
+	if (Index != INDEX_NONE)
+	{
+		Children.RemoveAt(Index);
+		return true;
+	}
+	return false;
+}
+
+bool FRedirectedPropertyNode::Remove(FName ChildPropertyName, FName ChildType, int32 ChildArrayIndex)
+{
+	const int32 Index = FindIndex(ChildPropertyName, ChildType, ChildArrayIndex);
+	if (Index != INDEX_NONE)
+	{
+		Children.RemoveAt(Index);
+		return true;
+	}
+	return false;
+}
+
+TSharedPtr<FRedirectedPropertyNode> FRedirectedPropertyNode::Find(const FPropertyPath& Path, int32 PathIndex) const
+{
+	check(PathIndex <= Path.GetNumProperties());
+	if (PathIndex == Path.GetNumProperties())
+	{
+		return SharedThis(const_cast<FRedirectedPropertyNode*>(this));
+	}
+	
+	const FPropertyInfo& ChildInfo = Path.GetPropertyInfo(PathIndex);
+	if (const TSharedPtr<FRedirectedPropertyNode> Child = Find(ChildInfo))
+	{
+		return Child->Find(Path, PathIndex + 1);
+	}
+	return {};
+}
+
+TSharedPtr<FRedirectedPropertyNode> FRedirectedPropertyNode::Find(const FPropertyInfo& ChildInfo) const
+{
+	const int32 Index = FindIndex(ChildInfo);
+	if (Index != INDEX_NONE)
+	{
+		return Children[Index];
+	}
+	return {};
+}
+
+TSharedPtr<FRedirectedPropertyNode> FRedirectedPropertyNode::Find(FName ChildPropertyName, FName ChildType, int32 ChildArrayIndex) const
+{
+	const int32 Index = FindIndex(ChildPropertyName, ChildType, ChildArrayIndex);
+	if (Index != INDEX_NONE)
+	{
+		return Children[Index];
+	}
+	return {};
+}
+
+bool FRedirectedPropertyNode::Move(const FPropertyPath& FromPath, const FPropertyPath& ToPath)
+{
+	if (TSharedPtr<FRedirectedPropertyNode> NodeToMove = Find(FromPath))
+	{
+		const TSharedPtr<FRedirectedPropertyNode> Added = FindOrAdd(ToPath);
+
+		// reparent children
+		Added->Children = MoveTemp(NodeToMove->Children);
+		for (const TSharedPtr<FRedirectedPropertyNode>& Child : Added->Children)
+		{
+			Child->Parent = Added;
+		}
+
+		do
+		{
+			const TSharedPtr<FRedirectedPropertyNode> ParentNode = NodeToMove->Parent.Pin();
+			if (ParentNode)
+			{
+				ParentNode->Remove(NodeToMove->PropertyName, NodeToMove->Type, NodeToMove->ArrayIndex);
+			}
+			NodeToMove = ParentNode;
+		} while (NodeToMove && NodeToMove->Children.IsEmpty());
+		return true;
+	}
+	return false;
+}
+
+int32 FRedirectedPropertyNode::FindIndex(const FPropertyInfo& ChildInfo) const
+{
+	return FindIndex(ChildInfo.Property->GetFName(), ChildInfo.Property->GetID(), ChildInfo.ArrayIndex);
+}
+
+int32 FRedirectedPropertyNode::FindIndex(FName ChildPropertyName, FName ChildType, int32 ChildArrayIndex) const
+{
+	return Children.IndexOfByPredicate([ChildPropertyName, ChildType, ChildArrayIndex](const TSharedPtr<FRedirectedPropertyNode>& Child)
+	{
+		if (Child->ArrayIndex != INDEX_NONE)
+		{
+			// a matching index will always match regardless of type and name
+			return Child->ArrayIndex == ChildArrayIndex;
+		}
+		return Child->Type == ChildType && Child->PropertyName == ChildPropertyName;
+	});
+}
+
+FArchetypeFixupPanel::FArchetypeFixupPanel(TConstArrayView<TObjectPtr<UObject>> Archetypes, EViewFlags InViewFlags)
+	: Instances(Archetypes)
+	, RedirectedPropertyTree(MakeShared<FRedirectedPropertyNode>())
+	, ViewFlags(InViewFlags)
+{
+	InitRedirectedPropertyTree();
+}
+
+int32 FArchetypeFixupPanel::Find(UObject* Value) const
+{
+	return Instances.Find(Value);
+}
+
+TSharedPtr<IDetailsView>& FArchetypeFixupPanel::GenerateDetailsView(bool bScrollbarOnLeft)
+{
+	FDetailsViewArgs DetailsViewArgs;
+	DetailsViewArgs.bUpdatesFromSelection = false;
+	DetailsViewArgs.bHideSelectionTip = true;
+	DetailsViewArgs.ExternalScrollbar = SAssignNew(LinkableScrollBar, SLinkableScrollBar);
+	DetailsViewArgs.ScrollbarAlignment = bScrollbarOnLeft ? HAlign_Left : HAlign_Right;
+	
+	FPropertyEditorModule& PropertyEditorModule = FModuleManager::GetModuleChecked<FPropertyEditorModule>("PropertyEditor");
+	DetailsView = PropertyEditorModule.CreateDetailView(DetailsViewArgs);
+
+	for (const UObject* Object : Instances)
+	{
+		if (HasViewFlag(EViewFlags::HideLooseProperties))
+		{
+			DetailsView->RegisterInstancedCustomPropertyLayout(Object->GetClass(), FOnGetDetailCustomizationInstance::CreateLambda([]()
+			{
+				return MakeShared<FHideLoosePropertiesCustomization>();
+			}));
+		}
+		else if (HasViewFlag(EViewFlags::AllowRemapLooseProperties))
+		{
+			DetailsView->RegisterInstancedCustomPropertyLayout(Object->GetClass(), FOnGetDetailCustomizationInstance::CreateLambda([DiffPanel = SharedThis(this)]()
+			{
+				return MakeShared<FArchetypeFixupDetailCustomization>(DiffPanel);
+			}));
+		}
+	}
+	
+	DetailsView->SetObjects(Instances, true);
+	return DetailsView;
+}
+
+void FArchetypeFixupPanel::SetDiffAgainstLeft(const TSharedPtr<FAsyncDetailViewDiff>& InDiffAgainstLeft)
+{
+	DiffAgainstLeft = InDiffAgainstLeft;
+}
+
+void FArchetypeFixupPanel::SetDiffAgainstRight(const TSharedPtr<FAsyncDetailViewDiff>& InDiffAgainstRight)
+{
+	DiffAgainstRight = InDiffAgainstRight;
+}
+
+TSharedPtr<FAsyncDetailViewDiff> FArchetypeFixupPanel::GetDiffAgainstLeft() const
+{
+	return DiffAgainstLeft.Pin();
+}
+
+TSharedPtr<FAsyncDetailViewDiff> FArchetypeFixupPanel::GetDiffAgainstRight() const
+{
+	return DiffAgainstRight.Pin();
+}
+
+bool FArchetypeFixupPanel::ShouldSplitterIgnoreRow(const TWeakPtr<FDetailTreeNode>& WeakDetailTreeNode) const
+{
+	if (const TSharedPtr<FDetailTreeNode> DetailTreeNode = WeakDetailTreeNode.Pin())
+	{
+		if (const TSharedPtr<IPropertyHandle> Handle = DetailTreeNode->CreatePropertyHandle())
+		{
+			if (MarkedForDelete.Contains(*Handle->CreateFPropertyPath()))
+			{
+				return true;
+			}
+		}
+	}
+	return false;
+}
+
+bool FArchetypeFixupPanel::AreAllConflictsRedirected() const
+{
+	bool bFoundConflict = false;
+	if (const TSharedPtr<FAsyncDetailViewDiff> Diff = DiffAgainstRight.Pin())
+	{
+		Diff->ForEach(ETreeTraverseOrder::PreOrder,
+		[this, &bFoundConflict](const TUniquePtr<FAsyncDetailViewDiff::DiffNodeType>& DiffNode)->ETreeTraverseControl
+		{
+			const TSharedPtr<FDetailTreeNode> TreeNode = DiffNode->ValueA.Pin();
+			if (DiffNode->DiffResult == ETreeDiffResult::MissingFromTree2 && TreeNode)
+			{
+				if (const TSharedPtr<IPropertyHandle> Handle = TreeNode->CreatePropertyHandle())
+				{
+					if (!MarkedForDelete.Contains(*Handle->CreateFPropertyPath()))
+					{
+						bFoundConflict = true;
+						return ETreeTraverseControl::Break;
+					}
+				}
+			}
+			return ETreeTraverseControl::Continue;
+		});
+	}
+	return !bFoundConflict;
+}
+
+void FArchetypeFixupPanel::AutoApplyMarkDeletedActions()
+{
+	const TSharedPtr<FAsyncDetailViewDiff> Diff = DiffAgainstRight.Pin();
+	if (!Diff)
+	{
+		return;
+	}
+
+	Diff->ForEach(ETreeTraverseOrder::PreOrder,
+		[this] (const TUniquePtr<FAsyncDetailViewDiff::DiffNodeType>& DiffNode)->ETreeTraverseControl
+		{
+			if (DiffNode->DiffResult == ETreeDiffResult::MissingFromTree2)
+			{
+				if (const TSharedPtr<FDetailTreeNode> LeftTreeNode = DiffNode->ValueA.Pin())
+				{
+					MarkForDelete(LeftTreeNode->GetPropertyPath());
+				}
+			}
+			
+			return ETreeTraverseControl::Continue;
+		});
+}
+
+bool FArchetypeFixupPanel::HasViewFlag(EViewFlags Flag)
+{
+	return static_cast<uint8>(Flag) & static_cast<uint8>(ViewFlags);
+}
+
+static void* ResolvePath(const FPropertyPath& Path, void* Value)
+{
+	for(int32 PathIndex = 0; PathIndex < Path.GetNumProperties(); ++PathIndex)
+	{
+		const FPropertyInfo& PropertyInfo = Path.GetPropertyInfo(PathIndex);
+		const FProperty* Property = PropertyInfo.Property.Get();
+		if (!Property)
+		{
+			return nullptr;
+		}
+
+		Value = Property->ContainerPtrToValuePtr<void>(Value, PropertyInfo.ArrayIndex != INDEX_NONE ? PropertyInfo.ArrayIndex : 0);
+
+		if (const FObjectProperty* AsObjectProperty = CastField<FObjectProperty>(Property))
+		{
+			Value = AsObjectProperty->GetObjectPropertyValue(Value);
+		}
+		else if (PathIndex + 1 < Path.GetNumProperties())
+		{
+			if (const FArrayProperty* AsArrayProperty = CastField<FArrayProperty>(Property))
+			{
+				FScriptArrayHelper Helper(AsArrayProperty, Value);
+				Value = Helper.GetElementPtr(Path.GetPropertyInfo(++PathIndex).ArrayIndex);
+			}
+			if (const FSetProperty* AsSetProperty = CastField<FSetProperty>(Property))
+			{
+				FScriptSetHelper Helper(AsSetProperty, Value);
+				Value = Helper.FindNthElementPtr(Path.GetPropertyInfo(++PathIndex).ArrayIndex);
+			}
+			if (const FMapProperty* AsMapProperty = CastField<FMapProperty>(Property))
+			{
+				FScriptMapHelper Helper(AsMapProperty, Value);
+				Value = Helper.FindNthValuePtr(Path.GetPropertyInfo(++PathIndex).ArrayIndex);
+			}
+		}
+	}
+	
+	return Value;
+}
+
+void FArchetypeFixupPanel::RedirectProperty(const FPropertyPath& From, const FPropertyPath& To)
+{
+	UArchetypeFixupUndoHandler* Snapshot = NewObject<UArchetypeFixupUndoHandler>();
+	Snapshot->Init(SharedThis(this));
+	GEditor->BeginTransaction(TEXT("ArchetypeFixupTool"), FText::Format(LOCTEXT("RedirectPropertyTransaction","Redirect {0} to {1}"), FText::FromString(From.ToString()), FText::FromString(To.ToString())), nullptr);
+	
+	if (const FPropertyPath* OriginalPath = OriginalPaths.Find(From))
+	{
+		if (To != *OriginalPath)
+		{
+			OriginalPaths.Add(To, *OriginalPath);
+		}
+		MarkedForDelete.Remove(*OriginalPath);
+		OriginalPaths.Remove(From);
+	}
+	else
+	{
+		OriginalPaths.Add(To, From);
+		MarkedForDelete.Remove(From);
+	}
+
+	if (To != From)
+	{
+		RedirectedPropertyTree->Move(From, To);
+	}
+	Snapshot->OnRedirect(From, To);
+		
+	for (UObject* Instance : Instances)
+	{
+		const FProperty* SourceProperty = From.GetLeafMostProperty().Property.Get();
+		const FProperty* DestinationProperty = To.GetLeafMostProperty().Property.Get();
+		if (!ensure(SourceProperty && DestinationProperty))
+		{
+			continue;
+		}
+		const void* Source = ResolvePath(From, Instance);
+		void* Destination = ResolvePath(To, Instance);
+		
+		if (!ensure(Source && Destination))
+		{
+			continue;
+		}
+	
+		FPropertyChangedEvent ChangeEvent(To.GetRootProperty().Property.Get(), EPropertyChangeType::ValueSet);
+		Instance->PreEditChange(ChangeEvent.Property);
+		
+		if (SourceProperty->SameType(DestinationProperty))
+		{
+			SourceProperty->CopyCompleteValue(Destination, Source);
+		}
+		else
+		{
+			FString ValueStr;
+			SourceProperty->ExportText_Direct(ValueStr, Source, nullptr, Instance, PPF_Copy);
+			DestinationProperty->ImportText_Direct(*ValueStr, Destination, Instance, PPF_Copy);
+		}
+		Instance->PostEditChangeProperty(ChangeEvent);
+	}
+
+	GEditor->EndTransaction();
+	
+	DetailsView->ForceRefresh();
+}
+
+void FArchetypeFixupPanel::OnRedirectProperty(FPropertyPath From, FPropertyPath To)
+{
+	RedirectProperty(From, To);
+}
+
+static void InitRedirectedPropertyTreeRec(const TSharedPtr<FRedirectedPropertyNode>& Node, FProperty* Property, void* Value);
+static void InitRedirectedPropertyTreeRec(const TSharedPtr<FRedirectedPropertyNode>& Node, UStruct* Struct, void* StructValue)
+{
+	for (FProperty* Property : TFieldRange<FProperty>(Struct))
+	{
+		if (Property->ArrayDim == 1)
+		{
+			if (UE::WasPropertySetBySerialization(Struct, StructValue, Property))
+			{
+				const TSharedPtr<FRedirectedPropertyNode>& ChildNode = Node->FindOrAdd(FPropertyInfo(Property));
+				void* Value = Property->ContainerPtrToValuePtr<void>(StructValue);
+				InitRedirectedPropertyTreeRec(ChildNode, Property, Value);
+			}
+		}
+		else
+		{
+			for (int32 StaticArrayIndex = 0; StaticArrayIndex < Property->ArrayDim; ++StaticArrayIndex)
+            {
+            	if (UE::WasPropertySetBySerialization(Struct, StructValue, Property, StaticArrayIndex))
+            	{
+            		const TSharedPtr<FRedirectedPropertyNode>& ChildNode = Node->FindOrAdd(FPropertyInfo(Property, StaticArrayIndex));
+            		void* Value = Property->ContainerPtrToValuePtr<void>(StructValue, StaticArrayIndex);
+            		InitRedirectedPropertyTreeRec(ChildNode, Property, Value);
+            	}
+            }
+		}
+	}
+}
+
+static void InitRedirectedPropertyTreeRec(const TSharedPtr<FRedirectedPropertyNode>& Node, FProperty* Property, void* Value)
+{
+	if (const FStructProperty* AsStructProperty = CastField<FStructProperty>(Property))
+	{
+		InitRedirectedPropertyTreeRec(Node, AsStructProperty->Struct, Value);
+	}
+	else if (const FArrayProperty* AsArrayProperty = CastField<FArrayProperty>(Property))
+	{
+		FScriptArrayHelper Array(AsArrayProperty, Value);
+		for (int32 ArrayIndex = 0; ArrayIndex < Array.Num(); ++ArrayIndex)
+		{
+			const TSharedPtr<FRedirectedPropertyNode>& ChildNode = Node->FindOrAdd(FPropertyInfo(AsArrayProperty->Inner, ArrayIndex));
+			InitRedirectedPropertyTreeRec(ChildNode, AsArrayProperty->Inner, Array.GetElementPtr(ArrayIndex));
+		}
+	}
+	else if (const FSetProperty* AsSetProperty = CastField<FSetProperty>(Property))
+	{
+		FScriptSetHelper Set(AsSetProperty, Value);
+		for (FScriptSetHelper::FIterator Itr = Set.CreateIterator(); Itr; ++Itr)
+		{
+			const TSharedPtr<FRedirectedPropertyNode>& ChildNode = Node->FindOrAdd(FPropertyInfo(AsSetProperty->ElementProp, Itr.GetLogicalIndex()));
+			InitRedirectedPropertyTreeRec(ChildNode, AsSetProperty->ElementProp, Set.GetElementPtr(Itr));
+		}
+	}
+	else if (const FMapProperty* AsMapProperty = CastField<FMapProperty>(Property))
+	{
+		FScriptMapHelper Map(AsMapProperty, Value);
+		for (FScriptMapHelper::FIterator Itr = Map.CreateIterator(); Itr; ++Itr)
+		{
+			const TSharedPtr<FRedirectedPropertyNode>& KeyNode = Node->FindOrAdd(FPropertyInfo(AsMapProperty->KeyProp, Itr.GetLogicalIndex()));
+			InitRedirectedPropertyTreeRec(KeyNode, AsMapProperty->KeyProp, Map.GetKeyPtr(Itr));
+			const TSharedPtr<FRedirectedPropertyNode>& ValNode = Node->FindOrAdd(FPropertyInfo(AsMapProperty->ValueProp, Itr.GetLogicalIndex()));
+			InitRedirectedPropertyTreeRec(ValNode, AsMapProperty->ValueProp, Map.GetValuePtr(Itr));
+		}
+	}
+}
+
+void FArchetypeFixupPanel::InitRedirectedPropertyTree()
+{
+	InitRedirectedPropertyTreeRec(RedirectedPropertyTree, Instances[0]->GetClass(), Instances[0]);
+}
+
+void UArchetypeFixupUndoHandler::Init(const TSharedRef<FArchetypeFixupPanel>& Panel)
+{
+	SetFlags(RF_Transactional);
+	ArchetypePanel = Panel;
+	OriginalPaths = Panel->OriginalPaths;
+	MarkedForDelete = Panel->MarkedForDelete;
+}
+
+void UArchetypeFixupUndoHandler::OnRedirect(const FPropertyPath& From, const FPropertyPath& To)
+{
+	if (const TSharedPtr<FArchetypeFixupPanel> Panel = ArchetypePanel.Pin())
+	{
+		RedirectFrom = From;
+		RedirectTo = To;
+		++ChangeNum;
+	}
+	Modify();
+}
+
+void UArchetypeFixupUndoHandler::PostEditUndo()
+{
+	if (const TSharedPtr<FArchetypeFixupPanel> Panel = ArchetypePanel.Pin())
+	{
+		if (RedirectTo != RedirectFrom)
+		{
+			Panel->RedirectedPropertyTree->Move(RedirectTo, RedirectFrom);
+			Swap(RedirectTo, RedirectFrom);
+		}
+		
+		Swap(Panel->OriginalPaths, OriginalPaths);
+		Swap(Panel->MarkedForDelete, MarkedForDelete);
+		Panel->DetailsView->ForceRefresh();
+	}
+}
+
+bool FArchetypeFixupPanel::IsInRedirectedPropertyTree(const FPropertyPath& Path) const
+{
+	return RedirectedPropertyTree->Find(Path).IsValid();
+}
+
+const FPropertyPath& FArchetypeFixupPanel::GetOriginalPath(const FPropertyPath& Path) const
+{
+	if (const FPropertyPath* Found = OriginalPaths.Find(Path))
+	{
+		return *Found;
+	}
+	return Path;
+}
+
+void FArchetypeFixupPanel::MarkForDelete(const FPropertyPath& CurrentPath)
+{
+	// undo any existing redirection on this node
+	if (const FPropertyPath* OriginalPath = OriginalPaths.Find(CurrentPath))
+	{
+		// move this property back to it's original location before marking it for delete
+		const FPropertyPath PathCopy = *OriginalPath; // RedirectProperty will invalidate pointers. copy path by value so it doesn't get destroyed.
+		RedirectProperty(CurrentPath, PathCopy);
+		MarkedForDelete.Add(PathCopy);
+	}
+	else
+	{
+		MarkedForDelete.Add(CurrentPath);
+	}
+}
+
+void FArchetypeFixupPanel::OnMarkForDelete(FPropertyPath Path)
+{
+	MarkForDelete(Path);
+}
+
+#undef LOCTEXT_NAMESPACE
diff --git a/Engine/Plugins/Editor/ArchetypeFixupTool/Source/Private/ArchetypeFixupPanel.h b/Engine/Plugins/Editor/ArchetypeFixupTool/Source/Private/ArchetypeFixupPanel.h
new file mode 100644
--- /dev/null
+++ b/Engine/Plugins/Editor/ArchetypeFixupTool/Source/Private/ArchetypeFixupPanel.h
@@ -0,0 +1,129 @@
+// Copyright Epic Games, Inc. All Rights Reserved.
+
+#pragma once
+
+#include "CoreMinimal.h"
+#include "AsyncDetailViewDiff.h"
+#include "UObject/Object.h"
+#include "Templates/SharedPointer.h"
+#include "ArchetypeFixupPanel.generated.h"
+
+class SLinkableScrollBar;
+
+struct FRedirectedPropertyNode : public TSharedFromThis<FRedirectedPropertyNode>
+{
+	FRedirectedPropertyNode() = default;
+	FRedirectedPropertyNode(const FRedirectedPropertyNode& Other);
+	FRedirectedPropertyNode(const FPropertyInfo& Info, const TWeakPtr<FRedirectedPropertyNode>& Parent);
+	FRedirectedPropertyNode(FName InPropertyName, FName InType, int32 InArrayIndex, const TWeakPtr<FRedirectedPropertyNode>& InParent);
+	TSharedPtr<FRedirectedPropertyNode> FindOrAdd(const FPropertyPath& Path, int32 PathIndex = 0);
+	TSharedPtr<FRedirectedPropertyNode> FindOrAdd(const FPropertyInfo& ChildInfo);
+	TSharedPtr<FRedirectedPropertyNode> FindOrAdd(FName ChildPropertyName, FName ChildType, int32 ChildArrayIndex = 0);
+	bool Remove(const FPropertyPath& Path, int32 PathIndex = 0);
+	bool Remove(const FPropertyInfo& ChildInfo);
+	bool Remove(FName ChildPropertyName, FName ChildType, int32 ChildArrayIndex = 0);
+	TSharedPtr<FRedirectedPropertyNode> Find(const FPropertyPath& Path, int32 PathIndex = 0) const;
+	TSharedPtr<FRedirectedPropertyNode> Find(const FPropertyInfo& ChildInfo) const;
+	TSharedPtr<FRedirectedPropertyNode> Find(FName ChildPropertyName, FName ChildType, int32 ChildArrayIndex = 0) const;
+	bool Move(const FPropertyPath& FromPath, const FPropertyPath& ToPath);
+	int32 FindIndex(const FPropertyInfo& ChildInfo) const;
+	int32 FindIndex(FName ChildPropertyName, FName ChildType, int32 ChildArrayIndex = 0) const;
+
+	FName PropertyName;
+	FName Type;
+	int32 ArrayIndex;
+	TWeakPtr<FRedirectedPropertyNode> Parent;
+	TArray<TSharedPtr<FRedirectedPropertyNode>> Children;
+};
+
+class FArchetypeFixupPanel : public TSharedFromThis<FArchetypeFixupPanel>
+{
+public:
+	enum class EViewFlags : uint8
+	{
+		None = 0,
+		HideLooseProperties = (1 << 0), // hide properties with isLoose metadata set to true
+		IncludeOnlySetBySerialization = (1 << 1), // hide properties that weren't set by serialization
+		AllowRemapLooseProperties = (1 << 3),
+		ReadonlyValues = (1 << 2),
+
+		// displays only properties found in the property bag, allow remapping, and disallow value edits
+		DefaultLeftPanel = IncludeOnlySetBySerialization | AllowRemapLooseProperties | ReadonlyValues,
+		// display only properties found in latest version of the class but allow value edits
+		DefaultRightPanel = HideLooseProperties,
+	};
+	
+	FArchetypeFixupPanel(TConstArrayView<TObjectPtr<UObject>> Archetypes, EViewFlags ViewFlags);
+	
+	int32 Find(UObject* Value) const;
+	TSharedPtr<IDetailsView>& GenerateDetailsView(bool bScrollbarOnLeft = false);
+	
+	void SetDiffAgainstLeft(const TSharedPtr<FAsyncDetailViewDiff>& DiffAgainstLeft);
+	void SetDiffAgainstRight(const TSharedPtr<FAsyncDetailViewDiff>& DiffAgainstRight);
+	TSharedPtr<FAsyncDetailViewDiff> GetDiffAgainstLeft() const;
+	TSharedPtr<FAsyncDetailViewDiff> GetDiffAgainstRight() const;
+
+	bool ShouldSplitterIgnoreRow(const TWeakPtr<FDetailTreeNode>& DetailTreeNode) const;
+	bool AreAllConflictsRedirected() const;
+	const FPropertyPath& GetOriginalPath(const FPropertyPath& Path) const;
+	
+	void MarkForDelete(const FPropertyPath& Path);
+	// pass-by-copy version for delegates. Use MarkForDelete when possible
+	void OnMarkForDelete(FPropertyPath Path);
+	// mark all conflicted properties for delete (FixupMode only)
+	void AutoApplyMarkDeletedActions();
+
+	// the redirected property tree keeps track of which properties in the archetype were either set by a property bag during serialization or
+	// redirected to from a floating property. The members of the tree are visible in the left panel.
+	bool IsInRedirectedPropertyTree(const FPropertyPath& Path) const;
+
+	bool HasViewFlag(EViewFlags Flag);
+
+	// Initialized by SArchetypeFixupTool
+	TSharedPtr<IDetailsView> DetailsView;
+	TSharedPtr<SLinkableScrollBar> LinkableScrollBar;
+
+private:
+	friend class FArchetypeFixupSpecification; // for access to Redirects
+	friend class FArchetypeFixupDetailNodeBuilder;
+	friend class UArchetypeFixupUndoHandler;
+	
+	void RedirectProperty(const FPropertyPath& From, const FPropertyPath& To);
+	// pass-by-copy version for delegates. Use RedirectProperty when possible
+	void OnRedirectProperty(FPropertyPath From, FPropertyPath To);
+	
+	void InitRedirectedPropertyTree();
+	
+	TArray<TObjectPtr<UObject>> Instances; // stores either Archetype property bag or just regular objects
+
+	// the redirected property tree keeps track of which properties in the archetype were either set by a property bag during serialization or
+	// redirected to from a floating property. The members of the tree are visible in the left panel.
+	TSharedPtr<FRedirectedPropertyNode> RedirectedPropertyTree;
+	
+	TMap<FPropertyPath, FPropertyPath> OriginalPaths;
+	TSet<FPropertyPath> MarkedForDelete;
+	
+	TWeakPtr<FAsyncDetailViewDiff> DiffAgainstLeft;
+	TWeakPtr<FAsyncDetailViewDiff> DiffAgainstRight;
+	EViewFlags ViewFlags = EViewFlags::None;
+};
+
+UCLASS()
+class UArchetypeFixupUndoHandler : public UObject
+{
+public:
+	GENERATED_BODY()
+	void Init(const TSharedRef<FArchetypeFixupPanel>& Panel);
+	void OnRedirect(const FPropertyPath& From, const FPropertyPath& To);
+	virtual void PostEditUndo() override;
+	
+	TWeakPtr<FArchetypeFixupPanel> ArchetypePanel;
+	
+	TMap<FPropertyPath, FPropertyPath> OriginalPaths;
+	TSet<FPropertyPath> MarkedForDelete;
+	FPropertyPath RedirectFrom;
+	FPropertyPath RedirectTo;
+	
+	UPROPERTY()
+	int32 ChangeNum = 0;
+};
diff --git a/Engine/Plugins/Editor/ArchetypeFixupTool/Source/Private/ArchetypeFixupTool.cpp b/Engine/Plugins/Editor/ArchetypeFixupTool/Source/Private/ArchetypeFixupTool.cpp
new file mode 100644
--- /dev/null
+++ b/Engine/Plugins/Editor/ArchetypeFixupTool/Source/Private/ArchetypeFixupTool.cpp
@@ -0,0 +1,198 @@
+// Copyright Epic Games, Inc. All Rights Reserved.
+
+#include "ArchetypeFixupTool.h"
+
+#include "ArchetypeFixupPanel.h"
+#include "DetailTreeNode.h"
+#include "SlateOptMacros.h"
+#include "PropertyEditorModule.h"
+#include "SDetailsSplitter.h"
+#include "Widgets/Layout/LinkableScrollBar.h"
+#include "Widgets/Input/SButton.h"
+
+
+#define LOCTEXT_NAMESPACE "ArchetypeFixupTool"
+
+class FArchetypeFixupSpecification : public TTreeDiffSpecification<TWeakPtr<FDetailTreeNode>>
+{
+public:
+	using Super = TTreeDiffSpecification<TWeakPtr<FDetailTreeNode>>;
+	
+	FArchetypeFixupSpecification(const TSharedPtr<FArchetypeFixupPanel>& InLeftPanel, const TSharedPtr<FArchetypeFixupPanel>& InRightPanel)
+		: LeftPanel(InLeftPanel)
+		, RightPanel(InRightPanel)
+	{}
+
+	// when diffing, use the redirects to match properties so that renames are respected
+	virtual bool AreMatching(const TWeakPtr<FDetailTreeNode>& TreeNodeA, const TWeakPtr<FDetailTreeNode>& TreeNodeB) const override
+	{
+		const TSharedPtr<FDetailTreeNode> PinnedTreeNodeA = TreeNodeA.Pin();
+		const TSharedPtr<FDetailTreeNode> PinnedTreeNodeB = TreeNodeB.Pin();
+		if (!PinnedTreeNodeA || !PinnedTreeNodeB)
+		{
+			return PinnedTreeNodeA == PinnedTreeNodeB;
+		}
+
+		const TSharedPtr<IPropertyHandle> PropertyHandleA = PinnedTreeNodeA->CreatePropertyHandle();
+		const TSharedPtr<IPropertyHandle> PropertyHandleB = PinnedTreeNodeB->CreatePropertyHandle();
+		if (!PropertyHandleA || !PropertyHandleB)
+		{
+			// category nodes
+			return PinnedTreeNodeA->GetNodeName() == PinnedTreeNodeB->GetNodeName();
+		}
+
+		return PropertyHandleA->GetProperty() == PropertyHandleB->GetProperty();
+	}
+
+	virtual bool ShouldMatchByValue(const TWeakPtr<FDetailTreeNode>& TreeNode) const override
+	{
+		return false;
+	}
+
+	virtual bool ShouldInheritEqualFromChildren(const TWeakPtr<FDetailTreeNode>& TreeNodeA, const TWeakPtr<FDetailTreeNode>& TreeNodeB) const override
+	{
+		return true;
+	}
+	
+	const TWeakPtr<FArchetypeFixupPanel> LeftPanel;
+	const TWeakPtr<FArchetypeFixupPanel> RightPanel;
+};
+
+BEGIN_SLATE_FUNCTION_BUILD_OPTIMIZATION
+void SArchetypeFixupTool::Construct(const FArguments& InArgs)
+{
+	Panels[0] = MakeShared<FArchetypeFixupPanel>(InArgs._Archetypes, FArchetypeFixupPanel::EViewFlags::DefaultLeftPanel);
+	Panels[1] = MakeShared<FArchetypeFixupPanel>(InArgs._Archetypes, FArchetypeFixupPanel::EViewFlags::DefaultRightPanel);
+	
+	ChildSlot
+	[
+		SNew(SVerticalBox)
+		+SVerticalBox::Slot()
+		.AutoHeight()
+		.VAlign(VAlign_Top)
+		[
+			SNew(SHorizontalBox)
+			+SHorizontalBox::Slot()
+			.AutoWidth()
+			.HAlign(HAlign_Left)
+			[
+				SNew(SButton)
+				.Text(LOCTEXT("AutoMarkForDeletion", "Mark remaining conflicts for deletion"))
+				.OnClicked(this, &SArchetypeFixupTool::OnAutoMarkForDeletion)
+			]
+		]
+		+SVerticalBox::Slot()
+		[
+			SAssignNew(Splitter, SDetailsSplitter)
+			.RowHighlightColor_Static(&SArchetypeFixupTool::GetRowHighlightColor)
+		]
+		+SVerticalBox::Slot()
+		.HAlign(HAlign_Right)
+		.VAlign(VAlign_Center)
+		.AutoHeight()
+		[
+			SNew(SButton)
+			.Text(LOCTEXT("ConfirmButton","Confirm"))
+			.OnClicked(this, &SArchetypeFixupTool::OnConfirmClicked)
+			.IsEnabled(this, &SArchetypeFixupTool::IsResolved)
+		]
+	];
+}
+
+void SArchetypeFixupTool::Tick(const FGeometry& AllottedGeometry, const double InCurrentTime, const float InDeltaTime)
+{
+	PanelDiff->Tick();
+}
+
+END_SLATE_FUNCTION_BUILD_OPTIMIZATION
+
+void SArchetypeFixupTool::SetDockTab(const TSharedRef<SDockTab>& DockTab)
+{
+	OwningDockTab = DockTab;
+}
+
+FReply SArchetypeFixupTool::OnConfirmClicked() const
+{
+	if (const TSharedPtr<SDockTab> DockTab = OwningDockTab.Pin())
+	{
+		DockTab->RequestCloseTab();
+	}
+	return FReply::Handled();
+}
+
+void SArchetypeFixupTool::GenerateDetailsViews()
+{
+	constexpr int32 LeftIndex = 0;
+	constexpr int32 RightIndex = 1;
+	
+	// generate the panels
+	Panels[LeftIndex]->GenerateDetailsView(true);
+	Panels[RightIndex]->GenerateDetailsView(false);
+
+	// diff the panels
+	PanelDiff = MakeShared<FAsyncDetailViewDiff>(
+		Panels[LeftIndex]->DetailsView.ToSharedRef(), Panels[RightIndex]->DetailsView.ToSharedRef());
+	PanelDiff->SetDiffSpecification<FArchetypeFixupSpecification>(Panels[LeftIndex], Panels[RightIndex]);
+    
+	SLinkableScrollBar::LinkScrollBars(Panels[LeftIndex]->LinkableScrollBar.ToSharedRef(), Panels[RightIndex]->LinkableScrollBar.ToSharedRef(),
+		TAttribute<TArray<FVector2f>>::CreateRaw(PanelDiff.Get(), &FAsyncDetailViewDiff::GenerateScrollSyncRate)); // TODO: Make work w/ CreateShared
+	Panels[LeftIndex]->SetDiffAgainstRight(PanelDiff);
+	Panels[RightIndex]->SetDiffAgainstLeft(PanelDiff);
+
+	// add the panels to the splitter
+	for (TSharedPtr<FArchetypeFixupPanel> Panel : Panels)
+	{
+		Splitter->AddSlot(
+			SDetailsSplitter::Slot()
+			.DetailsView(Panel->DetailsView.ToSharedRef())
+			.DifferencesWithRightPanel(Panel.Get(), &FArchetypeFixupPanel::GetDiffAgainstRight)
+			.IsReadonly_Lambda([Panel = Panel.ToWeakPtr()](){return Panel.IsValid() ? Panel.Pin()->HasViewFlag(FArchetypeFixupPanel::EViewFlags::ReadonlyValues) : true;})
+			.ShouldIgnoreRow(Panel.Get(), &FArchetypeFixupPanel::ShouldSplitterIgnoreRow)
+		);
+	}
+}
+
+FLinearColor SArchetypeFixupTool::GetRowHighlightColor(const TUniquePtr<FAsyncDetailViewDiff::DiffNodeType>& DiffNode)
+{
+	switch (DiffNode->DiffResult)
+	{
+	case ETreeDiffResult::MissingFromTree1:
+		// adds are green
+		return FLinearColor(0.f, 1.f, 0.f, .5f);
+	case ETreeDiffResult::MissingFromTree2:
+		// removes/conflicts are over-saturated red
+		return FLinearColor(1.5f, 0.3f, 0.3f, 1.f);
+	case ETreeDiffResult::DifferentValues:
+		return FLinearColor(0.f, 1.f, 1.f, .8f);
+	case ETreeDiffResult::Identical:
+		return FLinearColor();
+	default:
+		check(false);
+		return FLinearColor();
+	}
+}
+
+bool SArchetypeFixupTool::IsResolved() const
+{
+	for (const TSharedPtr<FArchetypeFixupPanel>& Panel : Panels)
+	{
+		if (!Panel->AreAllConflictsRedirected())
+		{
+			return false;
+		}
+	}
+	return true;
+}
+
+FReply SArchetypeFixupTool::OnAutoMarkForDeletion() const
+{
+	for (const TSharedPtr<FArchetypeFixupPanel>& Panel : Panels)
+	{
+		Panel->AutoApplyMarkDeletedActions();
+	}
+	return FReply::Handled();
+}
+
+
+#undef LOCTEXT_NAMESPACE
+
diff --git a/Engine/Plugins/Editor/ArchetypeFixupTool/Source/Private/ArchetypeFixupToolModule.cpp b/Engine/Plugins/Editor/ArchetypeFixupTool/Source/Private/ArchetypeFixupToolModule.cpp
new file mode 100644
--- /dev/null
+++ b/Engine/Plugins/Editor/ArchetypeFixupTool/Source/Private/ArchetypeFixupToolModule.cpp
@@ -0,0 +1,48 @@
+// Copyright Epic Games, Inc. All Rights Reserved.
+
+#include "ArchetypeFixupToolModule.h"
+#include "ArchetypeFixupTool.h"
+
+#define LOCTEXT_NAMESPACE "ArchetypeFixupToolModule"
+
+static const FName ArchetypeFixupToolTabName = FName(TEXT("ArchetypeFixupTool"));
+
+void FArchetypeFixupToolModule::StartupModule()
+{
+}
+
+void FArchetypeFixupToolModule::ShutdownModule()
+{
+	
+}
+
+bool FArchetypeFixupToolModule::OpenArchetypeFixupTool() const
+{
+	TSharedPtr<SDockTab> DockTab = FGlobalTabmanager::Get()->TryInvokeTab(FTabId(ArchetypeFixupToolTabName));
+	if (!DockTab)
+	{
+		return false;
+	}
+	DockTab->DrawAttention();
+	return true;
+}
+
+TSharedRef<SDockTab> FArchetypeFixupToolModule::CreateArchetypeFixupTab(const FSpawnTabArgs& TabArgs, TConstArrayView<TObjectPtr<UObject>> Archetypes) const
+{
+	const TSharedRef<SArchetypeFixupTool> ArchetypeFixupTool = SNew(SArchetypeFixupTool)
+		.Archetypes(Archetypes);
+	TSharedRef<SDockTab> DockTab = SNew(SDockTab)
+		.TabRole(ETabRole::NomadTab)
+		[
+			ArchetypeFixupTool
+		];
+
+	ArchetypeFixupTool->SetDockTab(DockTab);
+	ArchetypeFixupTool->GenerateDetailsViews();
+	
+	return DockTab;
+}
+
+#undef LOCTEXT_NAMESPACE
+
+IMPLEMENT_MODULE(FArchetypeFixupToolModule, ArchetypeFixupTool)
diff --git a/Engine/Plugins/Editor/ArchetypeFixupTool/Source/Public/ArchetypeFixupTool.h b/Engine/Plugins/Editor/ArchetypeFixupTool/Source/Public/ArchetypeFixupTool.h
new file mode 100644
--- /dev/null
+++ b/Engine/Plugins/Editor/ArchetypeFixupTool/Source/Public/ArchetypeFixupTool.h
@@ -0,0 +1,45 @@
+// Copyright Epic Games, Inc. All Rights Reserved.
+
+#pragma once
+
+#include "CoreMinimal.h"
+#include "AsyncDetailViewDiff.h"
+#include "Widgets/SCompoundWidget.h"
+#include "Widgets/Docking/SDockTab.h"
+
+class SDetailsSplitter;
+class FArchetypeFixupPanel;
+struct FPropertySoftPath;
+class IStructureDetailsView;
+
+/**
+ * This tool diffs multiple property bags of one format against the same number of property bags of another format.
+ * 
+ */
+class ARCHETYPEFIXUPTOOL_API SArchetypeFixupTool : public SCompoundWidget
+{
+public:
+	SLATE_BEGIN_ARGS(SArchetypeFixupTool)
+	{}
+		SLATE_ARGUMENT(TConstArrayView<TObjectPtr<UObject>>, Archetypes)
+	SLATE_END_ARGS()
+
+	/** Constructs this widget with InArgs */
+	void Construct(const FArguments& InArgs);
+	virtual void Tick(const FGeometry& AllottedGeometry, const double InCurrentTime, const float InDeltaTime) override;
+	
+	void SetDockTab(const TSharedRef<SDockTab>& DockTab);
+	void GenerateDetailsViews();
+	static FLinearColor GetRowHighlightColor(const TUniquePtr<FAsyncDetailViewDiff::DiffNodeType>& DiffNode);
+	bool IsResolved() const;
+	FReply OnAutoMarkForDeletion() const;
+	
+private:
+	FReply OnConfirmClicked() const;
+	
+	TWeakPtr<SDockTab> OwningDockTab;
+	TStaticArray<TSharedPtr<FArchetypeFixupPanel>, 2> Panels;
+	TSharedPtr<FAsyncDetailViewDiff> PanelDiff;
+
+	TSharedPtr<SDetailsSplitter> Splitter;
+};
diff --git a/Engine/Plugins/Editor/ArchetypeFixupTool/Source/Public/ArchetypeFixupToolModule.h b/Engine/Plugins/Editor/ArchetypeFixupTool/Source/Public/ArchetypeFixupToolModule.h
new file mode 100644
--- /dev/null
+++ b/Engine/Plugins/Editor/ArchetypeFixupTool/Source/Public/ArchetypeFixupToolModule.h
@@ -0,0 +1,20 @@
+// Copyright Epic Games, Inc. All Rights Reserved.
+
+#pragma once
+#include "Modules/ModuleManager.h"
+#include "Widgets/Docking/SDockTab.h"
+
+class ARCHETYPEFIXUPTOOL_API FArchetypeFixupToolModule : public FDefaultModuleImpl
+{
+public:
+	virtual void StartupModule() override;
+	virtual void ShutdownModule() override;
+	bool OpenArchetypeFixupTool() const;
+	static FArchetypeFixupToolModule& Get()
+	{
+		return FModuleManager::LoadModuleChecked<FArchetypeFixupToolModule>("ArchetypeFixupTool");
+	}
+
+	// opens the fixup tool 
+	TSharedRef<SDockTab> CreateArchetypeFixupTab(const FSpawnTabArgs& TabArgs, TConstArrayView<TObjectPtr<UObject>> Archetypes) const;
+};
