[Tests] Proof of concept test for Property Bags and Archetypes
This CL introduces a ue console command that will construct a test object (UTestReportCardV<1..5>) and then upgrade it to a different version creating an archetype in the process. That archetype will then be displayed in a details panel.
Syntax:  TestVerseArchetypes <OldClassVersion> <NewClassVersion>
OldClassVersion and NewClassVersion must be integers between 1 and 5

diff --git a/Engine/Plugins/Editor/ArchetypeFixupTool/Source/Private/Tests/ArchetypeTests.cpp b/Engine/Plugins/Editor/ArchetypeFixupTool/Source/Private/Tests/ArchetypeTests.cpp
new file mode 100644
--- /dev/null
+++ b/Engine/Plugins/Editor/ArchetypeFixupTool/Source/Private/Tests/ArchetypeTests.cpp
@@ -0,0 +1,363 @@
+// Copyright Epic Games, Inc. All Rights Reserved.
+
+#include "ArchetypeTests.h"
+
+#include "CoreTypes.h"
+#include "ArchetypeFixupToolModule.h"
+#include "WorkspaceMenuStructure.h"
+#include "WorkspaceMenuStructureModule.h"
+#include "Templates/Function.h"
+#include "Async/Async.h"
+#include "Serialization/ObjectReader.h"
+#include "Serialization/ObjectWriter.h"
+#include "UObject/PropertyBagRepository.h"
+#include "UObject/ArchetypeUtils.h"
+#include "UObject/PropertyPathName.h"
+#include "UObject/UObjectThreadContext.h"
+#include "UObject/UObjectGlobals.h"
+#include "UObject/Package.h"
+#include "UnrealEngine.h"
+#include "DetailsViewArgs.h"
+#include "PropertyEditorModule.h"
+#include "Modules/ModuleManager.h"
+#include "IDetailsView.h"
+
+#define LOCTEXT_NAMESPACE "PropertyBagTests"
+
+static UTestReportCardV1* CreateTestObjectV1()
+{
+	UTestReportCardV1* Result = NewObject<UTestReportCardV1>();
+    Result->Name = TEXT("John Doe");
+    Result->Age = 11;
+    Result->Gender = StudentGender::Male;
+	Result->Grade = 82.875f;
+	return Result;
+}
+
+static UTestReportCardV2* CreateTestObjectV2()
+{
+	UTestReportCardV2* Result = NewObject<UTestReportCardV2>();
+    Result->Name = TEXT("John Doe");
+    Result->Age = 11;
+    Result->Gender = StudentGender::Male;
+	Result->GPA = 82.875f;
+	Result->MathGrade = 75.5f;
+	Result->MathNotes = {
+		TEXT("Lacks initiative"),
+	};
+	Result->ScienceGrade = 100.f;
+	Result->ScienceNotes = {
+		TEXT("I love having John in my class"),
+	};
+	Result->ArtGrade = 95.f;
+	Result->ArtNotes = {};
+	Result->EnglishGrade = 61.f;
+	Result->EnglishNotes = {
+		TEXT("John doesn't do his homework"),
+		TEXT("See me after class!"),
+	};
+	return Result;
+}
+
+static UTestReportCardV3* CreateTestObjectV3()
+{
+	UTestReportCardV3* Result = NewObject<UTestReportCardV3>();
+    Result->StudentName = TEXT("John Doe");
+    Result->Age = 11;
+    Result->Gender = TEXT("Male");
+	Result->GPA = 82.875f;
+    Result->MathReport = {
+    	.Name = TEXT("Math"),
+    	.Grade = 75.5f,
+    	.TeacherNotes = {
+    		TEXT("Lacks initiative"),
+    	},
+    };
+    Result->ScienceReport = {
+    	.Name = TEXT("Science"),
+    	.Grade = 100.f,
+    	.TeacherNotes = {
+    		TEXT("I love having John in my class"),
+    	},
+    };
+    Result->ArtReport = {
+    	.Name = TEXT("Art"),
+    	.Grade = 95.f,
+    	.TeacherNotes = {},
+    };
+    Result->EnglishReport = {
+    	.Name = TEXT("English"),
+    	.Grade = 61.f,
+    	.TeacherNotes = {
+    		TEXT("John doesn't do his homework"),
+    		TEXT("See me after class!"),
+    	},
+    };
+	return Result;
+}
+
+static UTestReportCardV4* CreateTestObjectV4()
+{
+	UTestReportCardV4* Result = NewObject<UTestReportCardV4>();
+    Result->StudentName = TEXT("John Doe");
+    Result->Age = 11;
+    Result->Gender = TEXT("Male");
+	Result->GPA = 82.875f;
+    Result->ClassGrades.Add(TEXT("Math"), {
+    	.Name = TEXT("Math"),
+    	.Grade = 75.5f,
+    	.TeacherNotes = {
+    		TEXT("Lacks initiative"),
+    	},
+    });
+    Result->ClassGrades.Add(TEXT("Science"), {
+    	.Name = TEXT("Science"),
+    	.Grade = 100.f,
+    	.TeacherNotes = {
+    		TEXT("I love having John in my class"),
+    	},
+    });
+    Result->ClassGrades.Add(TEXT("Art"), {
+    	.Name = TEXT("Art"),
+    	.Grade = 95.f,
+    	.TeacherNotes = {},
+    });
+    Result->ClassGrades.Add(TEXT("English"), {
+    	.Name = TEXT("English"),
+    	.Grade = 61.f,
+    	.TeacherNotes = {
+    		TEXT("John doesn't do his homework"),
+    		TEXT("See me after class!"),
+    	},
+    });
+	return Result;
+}
+
+static UTestReportCardV5* CreateTestObjectV5()
+{
+	UTestReportCardV5* Result = NewObject<UTestReportCardV5>();
+    Result->StudentName = TEXT("John Doe");
+    Result->Age = 11;
+    Result->Gender = TEXT("Male");
+	Result->GPA = 82.875f;
+    Result->ClassGrades.Add(TEXT("Math"), {
+    	.ClassName = TEXT("Math"),
+    	.Grade = 75.5f,
+    	.TeacherNotes = {
+    		TEXT("Lacks initiative"),
+    	},
+    });
+    Result->ClassGrades.Add(TEXT("Science"), {
+    	.ClassName = TEXT("Science"),
+    	.Grade = 100.f,
+    	.TeacherNotes = {
+    		TEXT("I love having John in my class"),
+    	},
+    });
+    Result->ClassGrades.Add(TEXT("Art"), {
+    	.ClassName = TEXT("Art"),
+    	.Grade = 95.f,
+    	.TeacherNotes = {},
+    });
+    Result->ClassGrades.Add(TEXT("English"), {
+    	.ClassName = TEXT("English"),
+    	.Grade = 61.f,
+    	.TeacherNotes = {
+    		TEXT("John doesn't do his homework"),
+    		TEXT("See me after class!"),
+    	},
+    });
+	return Result;
+}
+
+TArray<UObject*> CreateClassGradeTestObjects()
+{
+	return {
+		CreateTestObjectV1(),
+		CreateTestObjectV2(),
+		CreateTestObjectV3(),
+		CreateTestObjectV4(),
+		CreateTestObjectV5(),
+	};
+}
+
+// same as FObjectReader except it also calls UE::MarkPropertySetBySerialization for every property read
+class FArchetypeReader : public FObjectReader
+{
+public:
+	FArchetypeReader(UObject* Obj, const TArray<uint8>& InBytes)
+		: FObjectReader(InBytes)
+		, Object(Obj)
+	{
+		this->SetIsLoading(true);
+        this->SetIsPersistent(false);
+
+#if USE_STABLE_LOCALIZATION_KEYS
+        if (GIsEditor && !(ArPortFlags & (PPF_DuplicateVerbatim | PPF_DuplicateForPIE)))
+        {
+        	SetLocalizationNamespace(TextNamespaceUtil::EnsurePackageNamespace(Obj));
+        }
+#endif // USE_STABLE_LOCALIZATION_KEYS
+
+        Obj->Serialize(*this);
+	}
+
+	virtual void Serialize(void* Data, int64 Num) override
+	{
+		const UE::FPropertyPathName Path = FUObjectThreadContext::Get().GetSerializeContext()->SerializedPropertyPath;
+		if (!Path.IsEmpty())
+		{
+			if (!SeenPaths.Contains(Path))
+			{
+				UE::MarkPropertySetBySerialization(Object, Path);
+				SeenPaths.Add(Path);
+			}
+		}
+		FObjectReader::Serialize(Data, Num);
+	}
+private:
+	UObject* Object;
+	TSet<UE::FPropertyPathName> SeenPaths;
+};
+
+static UObject* GenerateTestArchetype(UObject* ObjectOld, UClass* NewClass)
+{
+	// simulate OldObject being serialized
+	TArray<uint8> SerializedData;
+	FObjectWriter(ObjectOld, SerializedData);
+	
+	// simulate OldObject being deserialized into NewObject creating a property bag
+	UObject* ObjectNew = NewObject<UObject>(GetTransientPackage(), NewClass);
+	{
+		FUObjectSerializeContext* LoadContext = FUObjectThreadContext::Get().GetSerializeContext();
+		const bool bRevertSerializeUnknownProperty = LoadContext->bSerializeUnknownProperty;
+		UObject* RevertSerializedObject = LoadContext->SerializedObject;
+		LoadContext->bSerializeUnknownProperty = true;
+		LoadContext->SerializedObject = ObjectNew;
+		
+		FObjectReader(ObjectNew, SerializedData);
+		
+		LoadContext->bSerializeUnknownProperty = bRevertSerializeUnknownProperty;
+		LoadContext->SerializedObject = RevertSerializedObject;
+	}
+	const UE::FPropertyBag* PropertyBag = UE::FPropertyBagRepository::Get().FindBag(ObjectNew);
+	
+	// construct archetype class
+	const UClass* ArchetypeClass = UE::CreatePropertyBagArchetypeClass(PropertyBag, NewClass, GetTransientPackage());
+
+	// construct archetype object
+	FStaticConstructObjectParameters Params(ArchetypeClass);
+	Params.SetFlags |= EObjectFlags::RF_Transactional;
+	UObject* ArchetypeObject = StaticConstructObject_Internal(Params);
+
+	// note: we could potentially use the archetypeClass CDO as the archetype object instead but it would make defaults behave differently
+	//UObject* ArchetypeObject = ArchetypeClass->GetDefaultObject();
+
+	// rerun serialization to read data into the archetype instead
+	FArchetypeReader(ArchetypeObject, SerializedData);
+	
+	return ArchetypeObject;
+}
+
+static void RunVerseArchetypeTestCommand(const TArray<FString>& Parameters)
+{
+	TArray<UObject*> ObjectUpgrades = CreateClassGradeTestObjects();
+
+	if (Parameters.Num() != 2)
+	{
+		UE_LOG(LogEngine, Display, TEXT("provide two integer parameters between 1 and %i"), ObjectUpgrades.Num());
+		return;
+	}
+
+	const int32 OldClassVersion = FCString::Atoi(*Parameters[0]) - 1;
+	const int32 NewClassVersion = FCString::Atoi(*Parameters[1]) - 1;
+	if (!ObjectUpgrades.IsValidIndex(OldClassVersion) || !ObjectUpgrades.IsValidIndex(NewClassVersion))
+	{
+		UE_LOG(LogEngine, Display, TEXT("provide two integer parameters between 1 and %i"), ObjectUpgrades.Num());
+		return;
+	}
+	
+	UObject* ObjectOld = ObjectUpgrades[OldClassVersion];
+
+	// simulate class changing and an archetype being generated
+	UClass* NewClass = ObjectUpgrades[NewClassVersion]->GetClass();
+	UObject* ArchetypeObject = GenerateTestArchetype(ObjectOld, NewClass);
+
+	FDetailsViewArgs DetailsViewArgs;
+	DetailsViewArgs.bUpdatesFromSelection = false;
+	DetailsViewArgs.bHideSelectionTip = true;
+	
+	FPropertyEditorModule& PropertyEditorModule = FModuleManager::GetModuleChecked<FPropertyEditorModule>("PropertyEditor");
+	TSharedRef<IDetailsView> DetailsView = PropertyEditorModule.CreateDetailView(DetailsViewArgs);
+	
+	DetailsView->SetObject(ArchetypeObject);
+	
+	const FName TabName = FName(TEXT("VerseArchetypesTest"));
+	FGlobalTabmanager::Get()->RegisterNomadTabSpawner(TabName, FOnSpawnTab::CreateLambda([DetailsView](const FSpawnTabArgs& TabArgs)
+	{
+		return SNew(SDockTab)
+		.TabRole(ETabRole::NomadTab)
+		[
+			DetailsView
+		];
+	}))
+		.SetDisplayName(LOCTEXT("VerseArchetypesTestTitle", "Verse Archetype Test"))
+		.SetGroup(WorkspaceMenu::GetMenuStructure().GetToolsCategory())
+		.SetMenuType(ETabSpawnerMenuType::Hidden);
+
+	const TSharedPtr<SDockTab> DockTab = FGlobalTabmanager::Get()->TryInvokeTab(FTabId(TabName));
+	DockTab->DrawAttention();
+}
+
+static void RunVerseArchetypeFixupCommand(const TArray<FString>& Parameters)
+{
+	TArray<UObject*> ObjectUpgrades = CreateClassGradeTestObjects();
+
+	if (Parameters.Num() != 2)
+	{
+		UE_LOG(LogEngine, Display, TEXT("provide two integer parameters between 1 and %i"), ObjectUpgrades.Num());
+		return;
+	}
+
+	const int32 OldClassVersion = FCString::Atoi(*Parameters[0]) - 1;
+	const int32 NewClassVersion = FCString::Atoi(*Parameters[1]) - 1;
+	if (!ObjectUpgrades.IsValidIndex(OldClassVersion) || !ObjectUpgrades.IsValidIndex(NewClassVersion))
+	{
+		UE_LOG(LogEngine, Display, TEXT("provide two integer parameters between 1 and %i"), ObjectUpgrades.Num());
+		return;
+	}
+	
+	UObject* ObjectOld = ObjectUpgrades[OldClassVersion];
+	
+	// simulate class changing and an archetype being generated
+	UClass* NewClass = ObjectUpgrades[NewClassVersion]->GetClass();
+	UObject* ArchetypeObject = GenerateTestArchetype(ObjectOld, NewClass);
+
+	const FName TabName = FName(TEXT("VerseArchetypesTestFixup"));
+	FGlobalTabmanager::Get()->RegisterNomadTabSpawner(TabName, FOnSpawnTab::CreateLambda([ArchetypeObject](const FSpawnTabArgs& TabArgs)
+	{
+		return FArchetypeFixupToolModule::Get().CreateArchetypeFixupTab(TabArgs, {ArchetypeObject});
+	}))
+		.SetDisplayName(LOCTEXT("VerseArchetypesTestTitle", "Verse Archetype Fixup"))
+		.SetGroup(WorkspaceMenu::GetMenuStructure().GetToolsCategory())
+		.SetMenuType(ETabSpawnerMenuType::Hidden);
+
+	const TSharedPtr<SDockTab> DockTab = FGlobalTabmanager::Get()->TryInvokeTab(FTabId(TabName));
+	DockTab->DrawAttention();
+}
+
+FAutoConsoleCommand TestArchetypesCommand(
+	TEXT("TestVerseArchetype"),
+	TEXT("syntax: TestVerseArchetype <OldClassVersion> <NewClassVersion>"),
+	FConsoleCommandWithArgsDelegate::CreateStatic(&RunVerseArchetypeTestCommand),
+	ECVF_Default
+);
+
+FAutoConsoleCommand TestArchetypeFixupCommand(
+	TEXT("TestVerseArchetypeFixup"),
+	TEXT("syntax: TestVerseArchetypeFixup <OldClassVersion> <NewClassVersion>"),
+	FConsoleCommandWithArgsDelegate::CreateStatic(&RunVerseArchetypeFixupCommand),
+	ECVF_Default
+);
+
+#undef LOCTEXT_NAMESPACE
diff --git a/Engine/Plugins/Editor/ArchetypeFixupTool/Source/Private/Tests/ArchetypeTests.h b/Engine/Plugins/Editor/ArchetypeFixupTool/Source/Private/Tests/ArchetypeTests.h
new file mode 100644
--- /dev/null
+++ b/Engine/Plugins/Editor/ArchetypeFixupTool/Source/Private/Tests/ArchetypeTests.h
@@ -0,0 +1,187 @@
+// Copyright Epic Games, Inc. All Rights Reserved.
+
+#pragma once
+
+#include "CoreMinimal.h"
+#include "UObject/Object.h"
+#include "ArchetypeTests.generated.h"
+
+UENUM(BlueprintType)
+enum class StudentGender : uint8
+{
+	Undefined,
+	Male,
+	Female,
+};
+
+UCLASS()
+class UTestReportCardV1 : public UObject
+{
+	GENERATED_BODY()
+	
+public:
+	UPROPERTY(EditAnywhere, Category = "Report Card")
+	FString Name;
+	
+	UPROPERTY(EditAnywhere, Category = "Report Card")
+	StudentGender Gender = StudentGender::Undefined;
+	
+	UPROPERTY(EditAnywhere, Category = "Report Card")
+	int32 Age = -1;
+	
+	UPROPERTY(EditAnywhere, Category = "Report Card")
+	float Grade = 100.f;
+};
+
+// Changing the type of Gender to FString
+// renaming Grade to GPA
+// adding several other properties
+UCLASS()
+class UTestReportCardV2 : public UObject
+{
+	GENERATED_BODY()
+	
+public:
+	UPROPERTY(EditAnywhere, Category = "Student Info")
+	FString Name;
+	
+	UPROPERTY(EditAnywhere, Category = "Student Info")
+	StudentGender Gender;
+	
+	UPROPERTY(EditAnywhere, Category = "Student Info")
+	int32 Age = -1;
+	
+	UPROPERTY(EditAnywhere, Category = "Report Card")
+	float GPA = 100.f;
+	
+	UPROPERTY(EditAnywhere, Category = "Report Card")
+	float MathGrade = 100.f;
+	
+	UPROPERTY(EditAnywhere, Category = "Report Card")
+	TArray<FString> MathNotes;
+	
+	UPROPERTY(EditAnywhere, Category = "Report Card")
+	float ScienceGrade = 100.f;
+	
+	UPROPERTY(EditAnywhere, Category = "Report Card")
+	TArray<FString> ScienceNotes;
+	
+	UPROPERTY(EditAnywhere, Category = "Report Card")
+	float ArtGrade = 100.f;
+	
+	UPROPERTY(EditAnywhere, Category = "Report Card")
+	TArray<FString> ArtNotes;
+	
+	UPROPERTY(EditAnywhere, Category = "Report Card")
+	float EnglishGrade = 100.f;
+	
+	UPROPERTY(EditAnywhere, Category = "Report Card")
+	TArray<FString> EnglishNotes;
+};
+
+USTRUCT()
+struct FClassGrade
+{
+	GENERATED_BODY()
+public:
+	UPROPERTY(EditAnywhere, Category = "Grade")
+	FString Name;
+	UPROPERTY(EditAnywhere, Category = "Grade")
+	float Grade = 100.f;
+	UPROPERTY(EditAnywhere, Category = "Grade")
+	TArray<FString> TeacherNotes;
+};
+
+// renaming Name to StudentName
+// moving float grades into sub-properties of FClassGrade
+UCLASS()
+class UTestReportCardV3 : public UObject
+{
+	GENERATED_BODY()
+	
+public:
+	UPROPERTY(EditAnywhere, Category = "Student Info")
+	FString StudentName;
+	
+	UPROPERTY(EditAnywhere, Category = "Student Info")
+	FString Gender;
+	
+	UPROPERTY(EditAnywhere, Category = "Student Info")
+	int32 Age = -1;
+	
+	UPROPERTY(EditAnywhere, Category = "Report Card")
+	float GPA = 100.f;
+	
+	UPROPERTY(EditAnywhere, Category = "Report Card")
+	FClassGrade MathReport;
+	
+	UPROPERTY(EditAnywhere, Category = "Report Card")
+	FClassGrade ScienceReport;
+	
+	UPROPERTY(EditAnywhere, Category = "Report Card")
+	FClassGrade ArtReport;
+	
+	UPROPERTY(EditAnywhere, Category = "Report Card")
+	FClassGrade EnglishReport;
+};
+
+// moving class grades into a map
+UCLASS()
+class UTestReportCardV4 : public UObject
+{
+	GENERATED_BODY()
+	
+public:
+	UPROPERTY(EditAnywhere, Category = "Student Info")
+	FString StudentName;
+	
+	UPROPERTY(EditAnywhere, Category = "Student Info")
+	FString Gender;
+	
+	UPROPERTY(EditAnywhere, Category = "Student Info")
+	int32 Age = -1;
+	
+	UPROPERTY(EditAnywhere, Category = "Report Card")
+	float GPA = 100.f;
+	
+	UPROPERTY(EditAnywhere, Category = "Report Card")
+	TMap<FString, FClassGrade> ClassGrades;
+};
+
+// changing class grade struct to include a ClassName
+USTRUCT()
+struct FClassGradeV2
+{
+	GENERATED_BODY()
+public:
+	UPROPERTY(EditAnywhere, Category = "Grade")
+	FString ClassName;
+	UPROPERTY(EditAnywhere, Category = "Grade")
+	float Grade = 100.f;
+	UPROPERTY(EditAnywhere, Category = "Grade")
+	TArray<FString> TeacherNotes;
+};
+
+UCLASS()
+class UTestReportCardV5 : public UObject
+{
+	GENERATED_BODY()
+	
+public:
+	UPROPERTY(EditAnywhere, Category = "Student Info")
+	FString StudentName;
+	
+	UPROPERTY(EditAnywhere, Category = "Student Info")
+	FString Gender;
+	
+	UPROPERTY(EditAnywhere, Category = "Student Info")
+	int32 Age = -1;
+	
+	UPROPERTY(EditAnywhere, Category = "Report Card")
+	float GPA = 100.f;
+	
+	UPROPERTY(EditAnywhere, Category = "Report Card")
+	TMap<FString, FClassGradeV2> ClassGrades;
+};
+
+TArray<UObject*> CreateClassGradeTestObjects();
\ No newline at end of file
diff --git a/Engine/Source/Runtime/CoreUObject/Internal/UObject/PropertyBagRepository.h b/Engine/Source/Runtime/CoreUObject/Internal/UObject/PropertyBagRepository.h
--- a/Engine/Source/Runtime/CoreUObject/Internal/UObject/PropertyBagRepository.h
+++ b/Engine/Source/Runtime/CoreUObject/Internal/UObject/PropertyBagRepository.h
@@ -87,8 +87,8 @@ public:
 	
 	// Accessors
 	bool HasBag(const UObjectBase* Owner) const;
-	FPropertyBag* FindBag(const UObjectBase* Owner);
-	const FPropertyBag* FindBag(const UObjectBase* Owner) const;
+	COREUOBJECT_API FPropertyBag* FindBag(const UObjectBase* Owner);
+	COREUOBJECT_API const FPropertyBag* FindBag(const UObjectBase* Owner) const;
 
 	bool HasArchetype(const UObjectBase* Owner) const;
 	UObject* FindArchetype(const UObjectBase* Owner);
