[CIS Fix] fix duplicate LOCTEXT keys

diff --git a/Engine/Plugins/Editor/ArchetypeFixupTool/Source/Private/ArchetypeFixupDetailCustomization.cpp b/Engine/Plugins/Editor/ArchetypeFixupTool/Source/Private/ArchetypeFixupDetailCustomization.cpp
--- a/Engine/Plugins/Editor/ArchetypeFixupTool/Source/Private/ArchetypeFixupDetailCustomization.cpp
+++ b/Engine/Plugins/Editor/ArchetypeFixupTool/Source/Private/ArchetypeFixupDetailCustomization.cpp
@@ -346,7 +346,7 @@ TSharedRef<SWidget> FArchetypeFixupDetailNodeBuilder::GeneratePropertyRedirectMe
 	UObject* FirstArchetype = Panel->Instances[0];
 	TSet<FPropertyPath> RedirectOptions = GetRedirectOptions(FirstArchetype->GetClass(), FirstArchetype);
 
-	MenuBuilder.BeginSection(NAME_None, LOCTEXT("MoveToMatchingRow", "Move"));
+	MenuBuilder.BeginSection(NAME_None, LOCTEXT("MoveProperty", "Move"));
 	{
 		for (const FPropertyPath& Option : RedirectOptions)
 		{
@@ -357,7 +357,7 @@ TSharedRef<SWidget> FArchetypeFixupDetailNodeBuilder::GeneratePropertyRedirectMe
 				if (Option.GetLeafMostProperty().Property->SameType(PropertyHandle->GetProperty()))
 				{
 					FText DisplayName = FText::FromString(Option.ToString());
-					FText Tooltip = FText::Format(LOCTEXT("ChangeToSameTypeTooltip", "Move property to '{0}'"), DisplayName);
+					FText Tooltip = FText::Format(LOCTEXT("MovePropertyTooltip", "Move property to '{0}'"), DisplayName);
 					MenuBuilder.AddMenuEntry(DisplayName, Tooltip, FSlateIcon()
 					, FUIAction(FExecuteAction::CreateSP(Panel.Get(), &FArchetypeFixupPanel::OnRedirectProperty, Path, Option))
 					, NAME_None
@@ -368,7 +368,7 @@ TSharedRef<SWidget> FArchetypeFixupDetailNodeBuilder::GeneratePropertyRedirectMe
 	}
 	MenuBuilder.EndSection();
 
-	MenuBuilder.BeginSection(NAME_None, LOCTEXT("ChangeToSameType", "Rename"));
+	MenuBuilder.BeginSection(NAME_None, LOCTEXT("RenameProperty", "Rename"));
 	{
 		for (const FPropertyPath& Option : RedirectOptions)
 		{
@@ -381,7 +381,7 @@ TSharedRef<SWidget> FArchetypeFixupDetailNodeBuilder::GeneratePropertyRedirectMe
 			if (Option.GetLeafMostProperty().Property->SameType(PropertyHandle->GetProperty()))
 			{
 				FText DisplayName = FText::FromString(Option.ToString());
-				FText Tooltip = FText::Format(LOCTEXT("ChangeToSameTypeTooltip", "Rename property to '{0}'"), DisplayName);
+				FText Tooltip = FText::Format(LOCTEXT("RenamePropertyTooltip", "Rename property to '{0}'"), DisplayName);
 				MenuBuilder.AddMenuEntry(DisplayName, Tooltip, FSlateIcon()
 				, FUIAction(FExecuteAction::CreateSP(Panel.Get(), &FArchetypeFixupPanel::OnRedirectProperty, Path, Option))
 				, NAME_None
@@ -402,7 +402,7 @@ TSharedRef<SWidget> FArchetypeFixupDetailNodeBuilder::GeneratePropertyRedirectMe
 		if (!Panel->MarkedForDelete.Contains(OriginalPath)) // check that it wasn't already marked as deleted
 		{
 			FText DisplayName = LOCTEXT("MarkForDeletion", "Mark For Deletion");
-			FText Tooltip = LOCTEXT("MarkForDeletion", "Mark this property for deletion");
+			FText Tooltip = LOCTEXT("MarkForDeletionTooltip", "Mark this property for deletion");
 			MenuBuilder.AddMenuEntry(DisplayName, Tooltip, FSlateIcon()
 						, FUIAction(FExecuteAction::CreateSP(Panel.Get(), &FArchetypeFixupPanel::OnMarkForDelete, Path))
 						, NAME_None
