[CIS FIX] More loctext duplicate key errors. My bad.

diff --git a/Engine/Plugins/Editor/ArchetypeFixupTool/Source/Private/Tests/ArchetypeTests.cpp b/Engine/Plugins/Editor/ArchetypeFixupTool/Source/Private/Tests/ArchetypeTests.cpp
--- a/Engine/Plugins/Editor/ArchetypeFixupTool/Source/Private/Tests/ArchetypeTests.cpp
+++ b/Engine/Plugins/Editor/ArchetypeFixupTool/Source/Private/Tests/ArchetypeTests.cpp
@@ -338,7 +338,7 @@ static void RunVerseArchetypeFixupCommand(const TArray<FString>& Parameters)
 	{
 		return FArchetypeFixupToolModule::Get().CreateArchetypeFixupTab(TabArgs, {ArchetypeObject});
 	}))
-		.SetDisplayName(LOCTEXT("VerseArchetypesTestTitle", "Verse Archetype Fixup"))
+		.SetDisplayName(LOCTEXT("VerseArchetypesFixupTestTitle", "Verse Archetype Fixup"))
 		.SetGroup(WorkspaceMenu::GetMenuStructure().GetToolsCategory())
 		.SetMenuType(ETabSpawnerMenuType::Hidden);
 
