[Feature] Fixup tool now marks redirected properties as transient and Hidden. Also made the Reset option revert value changes, transient changes, and metadata changes

diff --git a/Engine/Plugins/Editor/ArchetypeFixupTool/Source/Private/ArchetypeFixupDetailCustomization.cpp b/Engine/Plugins/Editor/ArchetypeFixupTool/Source/Private/ArchetypeFixupDetailCustomization.cpp
--- a/Engine/Plugins/Editor/ArchetypeFixupTool/Source/Private/ArchetypeFixupDetailCustomization.cpp
+++ b/Engine/Plugins/Editor/ArchetypeFixupTool/Source/Private/ArchetypeFixupDetailCustomization.cpp
@@ -303,7 +303,7 @@ int32 FArchetypeFixupDetailNodeBuilder::GetNameWidgetIndex() const
 	{
 		if (Panel->HasViewFlag(FArchetypeFixupPanel::EViewFlags::AllowRemapLooseProperties))
 		{
-			if (Panel->OriginalPaths.Contains(*PropertyHandle->CreateFPropertyPath()))
+			if (Panel->RevertInfo.Contains(*PropertyHandle->CreateFPropertyPath()))
 			{
 				return DisplayRedirectMenu;
 			}
diff --git a/Engine/Plugins/Editor/ArchetypeFixupTool/Source/Private/ArchetypeFixupPanel.cpp b/Engine/Plugins/Editor/ArchetypeFixupTool/Source/Private/ArchetypeFixupPanel.cpp
--- a/Engine/Plugins/Editor/ArchetypeFixupTool/Source/Private/ArchetypeFixupPanel.cpp
+++ b/Engine/Plugins/Editor/ArchetypeFixupTool/Source/Private/ArchetypeFixupPanel.cpp
@@ -319,7 +319,11 @@ void FArchetypeFixupPanel::AutoApplyMarkDeletedActions()
 			{
 				if (const TSharedPtr<FDetailTreeNode> LeftTreeNode = DiffNode->ValueA.Pin())
 				{
-					MarkForDelete(LeftTreeNode->GetPropertyPath());
+					const FPropertyPath Path = LeftTreeNode->GetPropertyPath();
+					if (Path.IsValid())
+					{
+						MarkForDelete(Path);
+					}
 				}
 			}
 			
@@ -377,37 +381,85 @@ void FArchetypeFixupPanel::RedirectProperty(const FPropertyPath& From, const FPr
 	UArchetypeFixupUndoHandler* Snapshot = NewObject<UArchetypeFixupUndoHandler>();
 	Snapshot->Init(SharedThis(this));
 	GEditor->BeginTransaction(TEXT("ArchetypeFixupTool"), FText::Format(LOCTEXT("RedirectPropertyTransaction","Redirect {0} to {1}"), FText::FromString(From.ToString()), FText::FromString(To.ToString())), nullptr);
+
+	FProperty* SourceProperty = From.GetLeafMostProperty().Property.Get();
+	check(SourceProperty);
+	FProperty* DestinationProperty = To.IsValid() ? To.GetLeafMostProperty().Property.Get() : nullptr;
+	FRevertInfo* ToRevertInfo = nullptr;
+	TOptional<FRevertInfo> FromRevertInfo;
 	
-	if (const FPropertyPath* OriginalPath = OriginalPaths.Find(From))
+	if (const FRevertInfo* Info = RevertInfo.Find(From))
 	{
-		if (To != *OriginalPath)
+		FromRevertInfo = *Info;
+		if (DestinationProperty)
+		{
+			if (DestinationProperty->HasAnyPropertyFlags(CPF_Transient) != Info->bWasTransient)
+            {
+            	// toggle transient flag if needed
+            	DestinationProperty->PropertyFlags ^= CPF_Transient;
+            }
+            if (!Info->bWasHidden)
+            {
+            	DestinationProperty->RemoveMetaData(TEXT("Hidden"));
+            }
+            DestinationProperty->RemoveMetaData(TEXT("Redirected"));
+		}
+		
+		
+		if (To.IsValid() && To != Info->OriginalPath)
 		{
-			OriginalPaths.Add(To, *OriginalPath);
+			TArray<uint8> OriginalValue;
+			
+			ToRevertInfo = &RevertInfo.Add(To, {
+				.OriginalPath = Info->OriginalPath,
+				.bWasTransient = SourceProperty->HasAnyPropertyFlags(CPF_Transient),
+				.bWasHidden = SourceProperty->HasMetaData(TEXT("Hidden"))
+			});
 		}
-		MarkedForDelete.Remove(*OriginalPath);
-		OriginalPaths.Remove(From);
+		MarkedForDelete.Remove(Info->OriginalPath);
+		RevertInfo.Remove(From);
 	}
 	else
 	{
-		OriginalPaths.Add(To, From);
-		MarkedForDelete.Remove(From);
+		if (To.IsValid())
+		{
+			ToRevertInfo = &RevertInfo.Add(To, {
+				.OriginalPath = From,
+				.bWasTransient = SourceProperty->HasAnyPropertyFlags(CPF_Transient)
+			});
+			MarkedForDelete.Remove(From);
+		}
 	}
-
+	
 	if (To != From)
 	{
-		RedirectedPropertyTree->Move(From, To);
+		if (To.IsValid())
+		{
+			RedirectedPropertyTree->Move(From, To);
+		}
+		if (SourceProperty->HasMetaData(TEXT("isLoose")))
+		{
+			SourceProperty->PropertyFlags |= CPF_Transient;
+			SourceProperty->SetMetaData(TEXT("Hidden"), TEXT("True"));
+			SourceProperty->SetMetaData(TEXT("Redirected"), TEXT("True"));
+		}
 	}
+
 	Snapshot->OnRedirect(From, To);
-		
+	
+	if (!DestinationProperty)
+	{
+		MarkedForDelete.Add(From);
+		GEditor->EndTransaction();
+        DetailsView->ForceRefresh();
+		return; // delete actions don't need data copied
+	}
+	
+
+	const uint8* FromRevertInfoItr = FromRevertInfo ? FromRevertInfo->OriginalValue.GetData() : nullptr;
 	for (UObject* Instance : Instances)
 	{
-		const FProperty* SourceProperty = From.GetLeafMostProperty().Property.Get();
-		const FProperty* DestinationProperty = To.GetLeafMostProperty().Property.Get();
-		if (!ensure(SourceProperty && DestinationProperty))
-		{
-			continue;
-		}
-		const void* Source = ResolvePath(From, Instance);
+		void* Source = ResolvePath(From, Instance);
 		void* Destination = ResolvePath(To, Instance);
 		
 		if (!ensure(Source && Destination))
@@ -417,6 +469,15 @@ void FArchetypeFixupPanel::RedirectProperty(const FPropertyPath& From, const FPr
 	
 		FPropertyChangedEvent ChangeEvent(To.GetRootProperty().Property.Get(), EPropertyChangeType::ValueSet);
 		Instance->PreEditChange(ChangeEvent.Property);
+
+		if (ToRevertInfo)
+		{
+			// cache the destination value so it can be reverted later
+			const int32 Size = DestinationProperty->ArrayDim * DestinationProperty->ElementSize;
+			ToRevertInfo->OriginalValue.AddZeroed(Size);
+			uint8* Buffer = ToRevertInfo->OriginalValue.GetData() + (ToRevertInfo->OriginalValue.Num() - Size);
+			DestinationProperty->CopyCompleteValue(Buffer, Destination);
+		}
 		
 		if (SourceProperty->SameType(DestinationProperty))
 		{
@@ -428,6 +489,13 @@ void FArchetypeFixupPanel::RedirectProperty(const FPropertyPath& From, const FPr
 			SourceProperty->ExportText_Direct(ValueStr, Source, nullptr, Instance, PPF_Copy);
 			DestinationProperty->ImportText_Direct(*ValueStr, Destination, Instance, PPF_Copy);
 		}
+		
+		if (FromRevertInfo)
+		{
+			// apply FromRevertInfo to From
+			SourceProperty->CopyCompleteValue(Source, FromRevertInfoItr);
+			FromRevertInfoItr += DestinationProperty->ArrayDim * DestinationProperty->ElementSize;
+		}
 		Instance->PostEditChangeProperty(ChangeEvent);
 	}
 
@@ -514,10 +582,10 @@ void FArchetypeFixupPanel::InitRedirectedPropertyTree()
 
 void UArchetypeFixupUndoHandler::Init(const TSharedRef<FArchetypeFixupPanel>& Panel)
 {
-	SetFlags(RF_Transactional);
 	ArchetypePanel = Panel;
-	OriginalPaths = Panel->OriginalPaths;
+	RevertInfo = Panel->RevertInfo;
 	MarkedForDelete = Panel->MarkedForDelete;
+	SetFlags(RF_Transactional);
 }
 
 void UArchetypeFixupUndoHandler::OnRedirect(const FPropertyPath& From, const FPropertyPath& To)
@@ -537,11 +605,14 @@ void UArchetypeFixupUndoHandler::PostEditUndo()
 	{
 		if (RedirectTo != RedirectFrom)
 		{
-			Panel->RedirectedPropertyTree->Move(RedirectTo, RedirectFrom);
+			if (RedirectTo.IsValid() && RedirectFrom.IsValid())
+			{
+				Panel->RedirectedPropertyTree->Move(RedirectTo, RedirectFrom);
+			}
 			Swap(RedirectTo, RedirectFrom);
 		}
 		
-		Swap(Panel->OriginalPaths, OriginalPaths);
+		Swap(Panel->RevertInfo, RevertInfo);
 		Swap(Panel->MarkedForDelete, MarkedForDelete);
 		Panel->DetailsView->ForceRefresh();
 	}
@@ -554,9 +625,9 @@ bool FArchetypeFixupPanel::IsInRedirectedPropertyTree(const FPropertyPath& Path)
 
 const FPropertyPath& FArchetypeFixupPanel::GetOriginalPath(const FPropertyPath& Path) const
 {
-	if (const FPropertyPath* Found = OriginalPaths.Find(Path))
+	if (const FRevertInfo* Found = RevertInfo.Find(Path))
 	{
-		return *Found;
+		return Found->OriginalPath;
 	}
 	return Path;
 }
@@ -564,16 +635,16 @@ const FPropertyPath& FArchetypeFixupPanel::GetOriginalPath(const FPropertyPath&
 void FArchetypeFixupPanel::MarkForDelete(const FPropertyPath& CurrentPath)
 {
 	// undo any existing redirection on this node
-	if (const FPropertyPath* OriginalPath = OriginalPaths.Find(CurrentPath))
+	if (const FRevertInfo* Found = RevertInfo.Find(CurrentPath))
 	{
 		// move this property back to it's original location before marking it for delete
-		const FPropertyPath PathCopy = *OriginalPath; // RedirectProperty will invalidate pointers. copy path by value so it doesn't get destroyed.
+		const FPropertyPath PathCopy = Found->OriginalPath; // RedirectProperty will invalidate pointers. copy path by value so it doesn't get destroyed.
 		RedirectProperty(CurrentPath, PathCopy);
-		MarkedForDelete.Add(PathCopy);
+		RedirectProperty(PathCopy, {});
 	}
 	else
 	{
-		MarkedForDelete.Add(CurrentPath);
+		RedirectProperty(CurrentPath, {});
 	}
 }
 
diff --git a/Engine/Plugins/Editor/ArchetypeFixupTool/Source/Private/ArchetypeFixupPanel.h b/Engine/Plugins/Editor/ArchetypeFixupTool/Source/Private/ArchetypeFixupPanel.h
--- a/Engine/Plugins/Editor/ArchetypeFixupTool/Source/Private/ArchetypeFixupPanel.h
+++ b/Engine/Plugins/Editor/ArchetypeFixupTool/Source/Private/ArchetypeFixupPanel.h
@@ -99,8 +99,16 @@ private:
 	// the redirected property tree keeps track of which properties in the archetype were either set by a property bag during serialization or
 	// redirected to from a floating property. The members of the tree are visible in the left panel.
 	TSharedPtr<FRedirectedPropertyNode> RedirectedPropertyTree;
+
+	struct FRevertInfo
+	{
+		TArray<uint8> OriginalValue;
+		FPropertyPath OriginalPath;
+		bool bWasTransient;
+		bool bWasHidden;
+	};
 	
-	TMap<FPropertyPath, FPropertyPath> OriginalPaths;
+	TMap<FPropertyPath, FRevertInfo> RevertInfo;
 	TSet<FPropertyPath> MarkedForDelete;
 	
 	TWeakPtr<FAsyncDetailViewDiff> DiffAgainstLeft;
@@ -119,7 +127,7 @@ public:
 	
 	TWeakPtr<FArchetypeFixupPanel> ArchetypePanel;
 	
-	TMap<FPropertyPath, FPropertyPath> OriginalPaths;
+	TMap<FPropertyPath, FArchetypeFixupPanel::FRevertInfo> RevertInfo;
 	TSet<FPropertyPath> MarkedForDelete;
 	FPropertyPath RedirectFrom;
 	FPropertyPath RedirectTo;
