[Feature] Archetype Support for subobjects
[Feature] PropertyBagRepository::FindArchetype is now the primary API for obtaining archetypes
[SubFeature] added OnPathChangedCallback to UE::FPropertyPathName so that archetypes can record which properties were deserialized

diff --git a/Engine/Plugins/Editor/ArchetypeFixupTool/Source/Private/Tests/ArchetypeTests.cpp b/Engine/Plugins/Editor/ArchetypeFixupTool/Source/Private/Tests/ArchetypeTests.cpp
--- a/Engine/Plugins/Editor/ArchetypeFixupTool/Source/Private/Tests/ArchetypeTests.cpp
+++ b/Engine/Plugins/Editor/ArchetypeFixupTool/Source/Private/Tests/ArchetypeTests.cpp
@@ -8,10 +8,7 @@
 #include "WorkspaceMenuStructureModule.h"
 #include "Templates/Function.h"
 #include "Async/Async.h"
-#include "Serialization/ObjectReader.h"
-#include "Serialization/ObjectWriter.h"
 #include "UObject/PropertyBagRepository.h"
-#include "UObject/ArchetypeUtils.h"
 #include "UObject/PropertyPathName.h"
 #include "UObject/UObjectThreadContext.h"
 #include "UObject/UObjectGlobals.h"
@@ -21,12 +18,28 @@
 #include "PropertyEditorModule.h"
 #include "Modules/ModuleManager.h"
 #include "IDetailsView.h"
+#include "UObject/CoreRedirects.h"
+#include "UObject/SavePackage.h"
+#include "Misc/PackageName.h"
+#include "HAL/FileManager.h"
 
 #define LOCTEXT_NAMESPACE "PropertyBagTests"
 
+static UPackage* CreateTestPackage(const UClass* Class)
+{
+	const FString UniqueAssetName = Class->GetName() + TEXT("_") + FGuid::NewGuid().ToString();
+    return CreatePackage(*(TEXT("/Temp/") + UniqueAssetName));
+}
+
+template<typename T>
+static UPackage* CreateTestPackage()
+{
+    return CreateTestPackage(T::StaticClass());
+}
+
 static UTestReportCardV1* CreateTestObjectV1()
 {
-	UTestReportCardV1* Result = NewObject<UTestReportCardV1>();
+	UTestReportCardV1* Result = NewObject<UTestReportCardV1>(CreateTestPackage<UTestReportCardV1>());
     Result->Name = TEXT("John Doe");
     Result->Age = 11;
     Result->Gender = StudentGender::Male;
@@ -36,7 +49,7 @@ static UTestReportCardV1* CreateTestObjectV1()
 
 static UTestReportCardV2* CreateTestObjectV2()
 {
-	UTestReportCardV2* Result = NewObject<UTestReportCardV2>();
+	UTestReportCardV2* Result = NewObject<UTestReportCardV2>(CreateTestPackage<UTestReportCardV2>());
     Result->Name = TEXT("John Doe");
     Result->Age = 11;
     Result->Gender = StudentGender::Male;
@@ -61,7 +74,7 @@ static UTestReportCardV2* CreateTestObjectV2()
 
 static UTestReportCardV3* CreateTestObjectV3()
 {
-	UTestReportCardV3* Result = NewObject<UTestReportCardV3>();
+	UTestReportCardV3* Result = NewObject<UTestReportCardV3>(CreateTestPackage<UTestReportCardV3>());
     Result->StudentName = TEXT("John Doe");
     Result->Age = 11;
     Result->Gender = TEXT("Male");
@@ -98,17 +111,50 @@ static UTestReportCardV3* CreateTestObjectV3()
 
 static UTestReportCardV4* CreateTestObjectV4()
 {
-	UTestReportCardV4* Result = NewObject<UTestReportCardV4>();
+	UTestReportCardV4* Result = NewObject<UTestReportCardV4>(CreateTestPackage<UTestReportCardV4>());
     Result->StudentName = TEXT("John Doe");
     Result->Age = 11;
     Result->Gender = TEXT("Male");
 	Result->GPA = 82.875f;
+	Result->ClassGradesArray = {
+		{
+			.Name = TEXT("Math"),
+			.Grade = 75.5f,
+			.TeacherNotes = {
+				TEXT("Lacks initiative"),
+			},
+			.MainTeacherName = FString(TEXT("Samantha"))
+		},
+		{
+			.Name = TEXT("Science"),
+			.Grade = 100.f,
+			.TeacherNotes = {
+				TEXT("I love having John in my class"),
+			},
+			.MainTeacherName = FString(TEXT("Mark"))
+		},
+		{
+			.Name = TEXT("Art"),
+			.Grade = 95.f,
+			.TeacherNotes = {},
+			.MainTeacherName = FString(TEXT("Frank"))
+		},
+		{
+			.Name = TEXT("English"),
+			.Grade = 61.f,
+			.TeacherNotes = {
+				TEXT("John doesn't do his homework"),
+				TEXT("See me after class!"),
+			}
+		}
+	};
     Result->ClassGrades.Add(TEXT("Math"), {
     	.Name = TEXT("Math"),
     	.Grade = 75.5f,
     	.TeacherNotes = {
     		TEXT("Lacks initiative"),
     	},
+		.MainTeacherName = FString(TEXT("Samantha"))
     });
     Result->ClassGrades.Add(TEXT("Science"), {
     	.Name = TEXT("Science"),
@@ -116,11 +162,13 @@ static UTestReportCardV4* CreateTestObjectV4()
     	.TeacherNotes = {
     		TEXT("I love having John in my class"),
     	},
+		.MainTeacherName = FString(TEXT("Mark"))
     });
     Result->ClassGrades.Add(TEXT("Art"), {
     	.Name = TEXT("Art"),
     	.Grade = 95.f,
     	.TeacherNotes = {},
+		.MainTeacherName = FString(TEXT("Frank"))
     });
     Result->ClassGrades.Add(TEXT("English"), {
     	.Name = TEXT("English"),
@@ -135,7 +183,7 @@ static UTestReportCardV4* CreateTestObjectV4()
 
 static UTestReportCardV5* CreateTestObjectV5()
 {
-	UTestReportCardV5* Result = NewObject<UTestReportCardV5>();
+	UTestReportCardV5* Result = NewObject<UTestReportCardV5>(CreateTestPackage<UTestReportCardV5>());
     Result->StudentName = TEXT("John Doe");
     Result->Age = 11;
     Result->Gender = TEXT("Male");
@@ -170,6 +218,62 @@ static UTestReportCardV5* CreateTestObjectV5()
 	return Result;
 }
 
+static UTestReportCardV6* CreateTestObjectV6()
+{
+	UTestReportCardV6* Result = NewObject<UTestReportCardV6>(CreateTestPackage<UTestReportCardV6>());
+	Result->StudentInfo = NewObject<UStudentInfoV1>(Result);
+	Result->StudentInfo->StudentName = TEXT("John Doe");
+	Result->StudentInfo->Age = 11;
+	Result->StudentInfo->Gender = TEXT("Male");
+	Result->StudentInfo->GradeYear = TEXT("6th Grade");
+	Result->Grade = 82.875f;
+	return Result;
+}
+
+static UTestReportCardV7* CreateTestObjectV7()
+{
+	UTestReportCardV7* Result = NewObject<UTestReportCardV7>(CreateTestPackage<UTestReportCardV7>());
+	Result->StudentInfo = NewObject<UStudentInfoV2>(Result);
+	Result->StudentInfo->StudentName = TEXT("John Doe");
+	Result->StudentInfo->Age = 11;
+	Result->StudentInfo->Gender = TEXT("Male");
+	Result->GPA = 82.875f;
+	return Result;
+}
+
+
+// hard coded mapping that allows us to treat different versions of a class as if they are the same class
+TMap<FString, FSoftObjectPath> GetExportClasses(const UClass* Class)
+{
+	TMap<FString, FSoftObjectPath> Result = {{TEXT("Self"), Class}};
+	
+	if (Class == UTestReportCardV6::StaticClass())
+	{
+		Result.Add({TEXT("StudentInfo"), UStudentInfoV1::StaticClass()});
+	}
+	else if (Class == UTestReportCardV7::StaticClass())
+	{
+		Result.Add({TEXT("StudentInfo"), UStudentInfoV2::StaticClass()});
+	}
+	
+	return Result;
+}
+
+TArray<FCoreRedirect> BuildTestRedirects(const UClass* OldClass, const UClass* NewClass)
+{
+	TMap<FString, FSoftObjectPath> OldExportClasses = GetExportClasses(OldClass);
+	TMap<FString, FSoftObjectPath> NewExportClasses = GetExportClasses(NewClass);
+	TArray<FCoreRedirect> Redirects;
+	for (const TPair<FString, FSoftObjectPath>& OldClassInfo : OldExportClasses)
+	{
+		if (const FSoftObjectPath* NewPath = NewExportClasses.Find(OldClassInfo.Key))
+		{
+			Redirects.Emplace(ECoreRedirectFlags::Type_Class, OldClassInfo.Value.GetAssetPath(), NewPath->GetAssetPath());
+		}
+	}
+	return Redirects;
+}
+
 TArray<UObject*> CreateClassGradeTestObjects()
 {
 	return {
@@ -178,86 +282,71 @@ TArray<UObject*> CreateClassGradeTestObjects()
 		CreateTestObjectV3(),
 		CreateTestObjectV4(),
 		CreateTestObjectV5(),
+		CreateTestObjectV6(),
+		CreateTestObjectV7(),
 	};
 }
 
-// same as FObjectReader except it also calls UE::MarkPropertySetBySerialization for every property read
-class FArchetypeReader : public FObjectReader
+FString SaveArchetypeObject(UObject* ObjectToSave)
 {
-public:
-	FArchetypeReader(UObject* Obj, const TArray<uint8>& InBytes)
-		: FObjectReader(InBytes)
-		, Object(Obj)
-	{
-		this->SetIsLoading(true);
-        this->SetIsPersistent(false);
+	UPackage* TempPackage = CreateTestPackage(ObjectToSave->GetClass());
 
-#if USE_STABLE_LOCALIZATION_KEYS
-        if (GIsEditor && !(ArPortFlags & (PPF_DuplicateVerbatim | PPF_DuplicateForPIE)))
-        {
-        	SetLocalizationNamespace(TextNamespaceUtil::EnsurePackageNamespace(Obj));
-        }
-#endif // USE_STABLE_LOCALIZATION_KEYS
+	FObjectDuplicationParameters DupParams = InitStaticDuplicateObjectParams(ObjectToSave, TempPackage);
+	UObject* DuplicateObjectToSave = StaticDuplicateObjectEx(DupParams);
+	DuplicateObjectToSave->SetFlags(RF_Public);
 
-        Obj->Serialize(*this);
-	}
+	FSavePackageArgs SaveArgs;
+	SaveArgs.TopLevelFlags = RF_Public;
+	SaveArgs.bWarnOfLongFilename = false;
+	SaveArgs.bSlowTask = false;
+
+	const FString PackagePathName = TempPackage->GetPathName();
+	const FString Filename = FPackageName::LongPackageNameToFilename(PackagePathName, FPackageName::GetAssetPackageExtension());
+	const FSavePackageResultStruct SavePackageResult = UPackage::Save(TempPackage, nullptr, *Filename, SaveArgs);
+	return SavePackageResult.IsSuccessful() && IFileManager::Get().FileExists(*Filename) ? Filename : FString();
+}
 
-	virtual void Serialize(void* Data, int64 Num) override
+static UObject* CreateArchetypeWithSubObjects(UObject* Owner)
+{
+	TArray<UObject*> SubObjects = {Owner};
+	GetObjectsWithOuter(Owner, SubObjects);
+	for (const UObject* SubObject : SubObjects)
 	{
-		const UE::FPropertyPathName Path = FUObjectThreadContext::Get().GetSerializeContext()->SerializedPropertyPath;
-		if (!Path.IsEmpty())
-		{
-			if (!SeenPaths.Contains(Path))
-			{
-				UE::MarkPropertySetBySerialization(Object, Path);
-				SeenPaths.Add(Path);
-			}
-		}
-		FObjectReader::Serialize(Data, Num);
+		UE::FPropertyBagRepository::Get().CreateArchetype(SubObject);
 	}
-private:
-	UObject* Object;
-	TSet<UE::FPropertyPathName> SeenPaths;
-};
+	return UE::FPropertyBagRepository::Get().FindArchetype(Owner);
+}
+
 
 static UObject* GenerateTestArchetype(UObject* ObjectOld, UClass* NewClass)
 {
-	// simulate OldObject being serialized
-	TArray<uint8> SerializedData;
-	FObjectWriter(ObjectOld, SerializedData);
+	// 1) simulate OldObject being serialized
+	FString Filename = SaveArchetypeObject(ObjectOld);
 	
-	// simulate OldObject being deserialized into NewObject creating a property bag
-	UObject* ObjectNew = NewObject<UObject>(GetTransientPackage(), NewClass);
+	// 2) simulate class changing
+	TArray<FCoreRedirect> Redirects = BuildTestRedirects(ObjectOld->GetClass(), NewClass);
+	FCoreRedirects::AddRedirectList(Redirects, TEXT("ArchetypeTests"));
+	
+	// 3) simulate OldObject being deserialized into ObjectNew. This will create a property bag.
+	UObject* ObjectNew = nullptr;
 	{
 		FUObjectSerializeContext* LoadContext = FUObjectThreadContext::Get().GetSerializeContext();
-		const bool bRevertTrackSerializedPropertyPath = LoadContext->bTrackSerializedPropertyPath;
-		const bool bRevertSerializeUnknownProperty = LoadContext->bSerializeUnknownProperty;
-		UObject* RevertSerializedObject = LoadContext->SerializedObject;
-		LoadContext->bTrackSerializedPropertyPath = true;
-		LoadContext->bSerializeUnknownProperty = true;
-		LoadContext->SerializedObject = ObjectNew;
-		
-		FObjectReader(ObjectNew, SerializedData);
+
+		// scoped flag changes to load context to support property bag generation
+		TGuardValue<bool> ScopedTrackSerializedPropertyPath(LoadContext->bTrackSerializedPropertyPath, true);
+		TGuardValue<bool> ScopedSerializeUnknownProperty(LoadContext->bSerializeUnknownProperty, true);
+		TGuardValue<bool> ScopedImpersonateProperties(LoadContext->bImpersonateProperties, true);
 		
-		LoadContext->bTrackSerializedPropertyPath = bRevertTrackSerializedPropertyPath;
-		LoadContext->bSerializeUnknownProperty = bRevertSerializeUnknownProperty;
-		LoadContext->SerializedObject = RevertSerializedObject;
+		UPackage* TempPackage = CreateTestPackage(NewClass);
+		LoadPackage(TempPackage, *Filename, LOAD_None);
+		ObjectNew = FindObject<UObject>(TempPackage, *ObjectOld->GetName());
 	}
-	const UE::FPropertyBag* PropertyBag = UE::FPropertyBagRepository::Get().FindBag(ObjectNew);
-	
-	// construct archetype class
-	const UClass* ArchetypeClass = UE::CreatePropertyBagArchetypeClass(PropertyBag, NewClass, GetTransientPackage());
-
-	// construct archetype object
-	FStaticConstructObjectParameters Params(ArchetypeClass);
-	Params.SetFlags |= EObjectFlags::RF_Transactional;
-	UObject* ArchetypeObject = StaticConstructObject_Internal(Params);
 
-	// note: we could potentially use the archetypeClass CDO as the archetype object instead but it would make defaults behave differently
-	//UObject* ArchetypeObject = ArchetypeClass->GetDefaultObject();
+	// 4) Construct an archetype from the NewObject
+	UObject* ArchetypeObject = CreateArchetypeWithSubObjects(ObjectNew);
 
-	// rerun serialization to read data into the archetype instead
-	FArchetypeReader(ArchetypeObject, SerializedData);
+	// revert step 2 so we don't break other tests
+	FCoreRedirects::RemoveRedirectList(Redirects, TEXT("ArchetypeTests"));
 	
 	return ArchetypeObject;
 }
diff --git a/Engine/Plugins/Editor/ArchetypeFixupTool/Source/Private/Tests/ArchetypeTests.h b/Engine/Plugins/Editor/ArchetypeFixupTool/Source/Private/Tests/ArchetypeTests.h
--- a/Engine/Plugins/Editor/ArchetypeFixupTool/Source/Private/Tests/ArchetypeTests.h
+++ b/Engine/Plugins/Editor/ArchetypeFixupTool/Source/Private/Tests/ArchetypeTests.h
@@ -90,6 +90,8 @@ public:
 	float Grade = 100.f;
 	UPROPERTY(EditAnywhere, Category = "Grade")
 	TArray<FString> TeacherNotes;
+	UPROPERTY(EditAnywhere, Category = "Grade")
+	FString MainTeacherName = FString(TEXT("Luther"));
 };
 
 // renaming Name to StudentName
@@ -144,6 +146,9 @@ public:
 	UPROPERTY(EditAnywhere, Category = "Report Card")
 	float GPA = 100.f;
 	
+	UPROPERTY(EditAnywhere, Category = "Report Card")
+	TArray<FClassGrade> ClassGradesArray;
+
 	UPROPERTY(EditAnywhere, Category = "Report Card")
 	TMap<FString, FClassGrade> ClassGrades;
 };
@@ -180,8 +185,74 @@ public:
 	UPROPERTY(EditAnywhere, Category = "Report Card")
 	float GPA = 100.f;
 	
+	UPROPERTY(EditAnywhere, Category = "Report Card", meta=(OriginalType="ClassGrade"))
+	TArray<FClassGradeV2> ClassGradesArray;
+
 	UPROPERTY(EditAnywhere, Category = "Report Card")
 	TMap<FString, FClassGradeV2> ClassGrades;
 };
 
+
+UCLASS()
+class UStudentInfoV1 : public UObject
+{
+	GENERATED_BODY()
+public:
+	UPROPERTY(EditAnywhere, Category = "Student Info")
+	FString StudentName;
+	
+	UPROPERTY(EditAnywhere, Category = "Student Info")
+	FString Gender;
+	
+	UPROPERTY(EditAnywhere, Category = "Student Info")
+	int32 Age = -1;
+
+	UPROPERTY(EditAnywhere, Category = "Student Info")
+	FString GradeYear;
+};
+
+UCLASS()
+class UTestReportCardV6 : public UObject
+{
+	GENERATED_BODY()
+    	
+public:
+	UPROPERTY(EditAnywhere, Instanced, Category = "Student Info")
+	UStudentInfoV1* StudentInfo;
+	
+	UPROPERTY(EditAnywhere, Category = "Report Card")
+	float Grade = 100.f;
+};
+
+UCLASS()
+class UStudentInfoV2 : public UObject
+{
+	GENERATED_BODY()
+public:
+	UPROPERTY(EditAnywhere, Category = "Student Info")
+	FString StudentName;
+	
+	UPROPERTY(EditAnywhere, Category = "Student Info")
+	FString Gender;
+	
+	UPROPERTY(EditAnywhere, Category = "Student Info")
+	int32 Age = -1;
+
+	UPROPERTY(EditAnywhere, Category = "Student Info")
+	FString GradeName;
+};
+
+UCLASS()
+class UTestReportCardV7 : public UObject
+{
+	GENERATED_BODY()
+    	
+public:
+	UPROPERTY(EditAnywhere, Instanced, Category = "Student Info")
+	UStudentInfoV2* StudentInfo;
+	
+	UPROPERTY(EditAnywhere, Category = "Report Card")
+	float GPA = 100.f;
+};
+
 TArray<UObject*> CreateClassGradeTestObjects();
\ No newline at end of file
diff --git a/Engine/Source/Runtime/CoreUObject/Internal/UObject/PropertyBagRepository.h b/Engine/Source/Runtime/CoreUObject/Internal/UObject/PropertyBagRepository.h
--- a/Engine/Source/Runtime/CoreUObject/Internal/UObject/PropertyBagRepository.h
+++ b/Engine/Source/Runtime/CoreUObject/Internal/UObject/PropertyBagRepository.h
@@ -13,6 +13,7 @@ namespace UE
 {
 
 class FPropertyBag;
+class FPropertyPathName;
 
 // Singleton class tracking property bag association with objects
 class FPropertyBagRepository
@@ -27,6 +28,7 @@ class FPropertyBagRepository
 	// TODO: Make private throughout and extend access permissions here or in wrapper classes? Don't want engine code modifying bags outside of serializers and details panels.
 	//friend UObjectBase;
 	//friend UStruct;
+	friend struct FScopedArchetypeLoad;
 
 private:
 	friend class FPropertyBagRepositoryLock;
@@ -66,7 +68,7 @@ public:
 	 * @param Owner			- Associated in world object.
 	 * @return				- Custom archetype object, UClass derived from associated bag.
 	 */
-	//UObject* CreateArchetype(UObjectBase* Owner);
+	COREUOBJECT_API UObject* CreateArchetype(const UObjectBase* Owner);
 
 	// TODO: Restrict property bag  destruction to within UObject::BeginDestroy() & FPropertyBagProperty destructor.
 	// Removes bag, archetype, and all associated data for this object.
@@ -90,9 +92,14 @@ public:
 	COREUOBJECT_API FPropertyBag* FindBag(const UObjectBase* Owner);
 	COREUOBJECT_API const FPropertyBag* FindBag(const UObjectBase* Owner) const;
 
-	bool HasArchetype(const UObjectBase* Owner) const;
-	UObject* FindArchetype(const UObjectBase* Owner);
-	const UObject* FindArchetype(const UObjectBase* Owner) const;
+	COREUOBJECT_API bool HasArchetype(const UObjectBase* Owner) const;
+	COREUOBJECT_API UObject* FindArchetype(const UObjectBase* Owner);
+	COREUOBJECT_API const UObject* FindArchetype(const UObjectBase* Owner) const;
+	
+	// query whether a property in an object was set when the object was deserialized
+	COREUOBJECT_API static bool WasPropertySetBySerialization(UObject* Object, const FPropertyPathName& Path);
+	// query whether a property in Struct was set when the struct was deserialized
+	COREUOBJECT_API static bool WasPropertySetBySerialization(const UStruct* Struct, const void* StructData, const FProperty* Property, int32 ArrayIndex = 0);
 
 private:
 	void Lock() const { CriticalSection.Lock(); }
diff --git a/Engine/Source/Runtime/CoreUObject/Private/UObject/ArchetypeUtils.cpp b/Engine/Source/Runtime/CoreUObject/Private/UObject/ArchetypeUtils.cpp
--- a/Engine/Source/Runtime/CoreUObject/Private/UObject/ArchetypeUtils.cpp
+++ b/Engine/Source/Runtime/CoreUObject/Private/UObject/ArchetypeUtils.cpp
@@ -12,6 +12,7 @@ namespace UE
 	using FWildcardPropertyPathName = FPropertyPathName;
 	
 	static FName Name_ValuesSetBySerialization(TEXT("_ValuesSetBySerialization"));
+	static const FName NAME_StructOrignalTypeMetadata(TEXT("OriginalType"));
 	
 	struct ResolvePropertyPathNameHelperParams
 	{
@@ -23,30 +24,26 @@ namespace UE
 		bool bAddIfNeeded = false;
 	};
 
-	static void BuildSegmentTypeFromProperty(const FProperty* Property, TArray<FName>& OutType)
+	static void BuildSegmentTypeFromProperty(const FProperty* Property, TStringBuilder<256>& OutType)
 	{
-#if false // TODO: @jordan.hoffmann when complete type info is finished use this branch instead
-		OutType.Add(Property->GetID());
-		if (const FStructProperty* AsStructProperty = CastField<FStructProperty>(Property))
+		auto Push = [&OutType](FName TypeName)
 		{
-#if WITH_EDITORONLY_DATA
-			if (const FString* TPSOverrideStructName = AsStructProperty->FindMetaData("TPSOverrideStructName"))
+			if (OutType.Len() > 0)
 			{
-				OutType.Add(FName(*TPSOverrideStructName));
+				OutType.AppendChar(' ');
 			}
-			else
-#endif
-			{
-				OutType.Add(AsStructProperty->Struct->GetFName());
-			}
-		}
-		else if (const FObjectProperty* AsObjectProperty = CastField<FObjectProperty>(Property))
+			FName::GetEntry(TypeName.GetDisplayIndex())->AppendNameToString(OutType);
+		};
+		
+#if false // TODO: @jordan.hoffmann when complete type info is finished use this branch instead
+		OutType.Add(Property->GetID());
+		if (const FObjectProperty* AsObjectProperty = CastField<FObjectProperty>(Property))
 		{
-			OutType.Add(AsObjectProperty->GetID());
+			Push(AsObjectProperty->GetID());
 		}
 		else if (const FEnumProperty* AsEnumProperty = CastField<FEnumProperty>(Property))
 		{
-			OutType.Add(AsEnumProperty->GetEnum()->GetFName());
+			Push(AsEnumProperty->GetEnum()->GetFName());
 		}
 		else if (const FArrayProperty* AsArrayProperty = CastField<FArrayProperty>(Property))
 		{
@@ -62,30 +59,21 @@ namespace UE
 			BuildSegmentTypeFromProperty(AsMapProperty->ValueProp, OutType);
 		}
 #else
-		OutType.Add(Property->GetID());
+		Push(Property->GetID());
 		if (const FStructProperty* AsStructProperty = CastField<FStructProperty>(Property))
 		{
-#if WITH_EDITORONLY_DATA
-			if (const FString* TPSOverrideStructName = AsStructProperty->FindMetaData("TPSOverrideStructName"))
-			{
-				OutType.Add(FName(*TPSOverrideStructName));
-			}
-			else
-#endif
-			{
-				OutType.Add(AsStructProperty->Struct->GetFName());
-			}
+			Push(AsStructProperty->Struct->GetFName());
 		}
 		else if (const FEnumProperty* AsEnumProperty = CastField<FEnumProperty>(Property))
 		{
-			OutType.Add(AsEnumProperty->GetEnum()->GetFName());
+			Push(AsEnumProperty->GetEnum()->GetFName());
 		}
 		else if (const FArrayProperty* AsArrayProperty = CastField<FArrayProperty>(Property))
 		{
 			if (AsArrayProperty->Inner->IsA<FEnumProperty>())
 			{
 				// enum paths currently don't recurse when they're in containers
-				OutType.Add(AsArrayProperty->Inner->GetClass()->GetFName());
+				Push(AsArrayProperty->Inner->GetClass()->GetFName());
 			}
 			else
 			{
@@ -95,24 +83,21 @@ namespace UE
 		else if (const FSetProperty* AsSetProperty = CastField<FSetProperty>(Property))
 		{
 			// sets currently don't recurse their element property types
-			OutType.Add(AsSetProperty->ElementProp->GetClass()->GetFName());
+			Push(AsSetProperty->ElementProp->GetClass()->GetFName());
 		}
 		else if (const FMapProperty* AsMapProperty = CastField<FMapProperty>(Property))
 		{
 			// maps currently don't recurse their key/value property types
-			OutType.Add(AsMapProperty->KeyProp->GetClass()->GetFName());
-			OutType.Add(AsMapProperty->ValueProp->GetClass()->GetFName());
+			Push(AsMapProperty->KeyProp->GetClass()->GetFName());
+			Push(AsMapProperty->ValueProp->GetClass()->GetFName());
 		}
 #endif
 	}
 	
 	static FName GetSegmentTypeFromProperty(const FProperty* Property)
 	{
-		TArray<FName> TypeList;
-		BuildSegmentTypeFromProperty(Property, TypeList);
-		TStringBuilder<512> Type;
-		Type.Join(TypeList, TEXT(' '));
-		
+		TStringBuilder<256> Type;
+		BuildSegmentTypeFromProperty(Property, Type);
 		return FName(Type);
 	}
 	
@@ -201,7 +186,8 @@ namespace UE
 		for (FProperty* Property : TFieldRange<FProperty>(Struct))
 		{
 			// find a property that matches the segment
-			if (Segment.Name != Property->GetFName() || Segment.Type != GetSegmentTypeFromProperty(Property))
+			const FName PropertySegmentType = GetSegmentTypeFromProperty(Property);
+			if (Segment.Name != Property->GetFName() || Segment.Type != PropertySegmentType)
 			{
 				continue;
 			}
@@ -224,7 +210,10 @@ namespace UE
 			}
 			else
 			{
-				++Params.CurPathIndex;
+				if (Segment.Index == INDEX_NONE)
+				{
+					++Params.CurPathIndex;
+				}
 				Params.ResultProperty = Property;
 				Params.Data = Property->ContainerPtrToValuePtr<void>(Params.Data);
 				return ResolvePropertyPathNameHelper(Params);
@@ -258,11 +247,17 @@ namespace UE
 	{
 		if (FStructProperty* AsStructProperty = CastField<FStructProperty>(Property))
 		{
-			const FName TPSOverrideStructName = AsStructProperty->Struct->GetFName();
+			const FString* OriginalType = nullptr;
+#if WITH_EDITORONLY_DATA
+			//@note: Transfer existing metadata over as we build the archetype from the struct or it owner, if any, this is useful for testing purposes
+			OriginalType = AsStructProperty->FindMetaData(NAME_StructOrignalTypeMetadata);
+			FField* OwnerField = OriginalType == nullptr ? AsStructProperty->Owner.ToField() : nullptr;
+			OriginalType = OwnerField ? OwnerField->FindMetaData(NAME_StructOrignalTypeMetadata) : OriginalType;
+#endif
+			const FName StructOriginalName = OriginalType ? FName(**OriginalType) : AsStructProperty->Struct->GetFName();
 			AsStructProperty->Struct = CreatePropertyBagArchetypeStructRec<UScriptStruct>(AsStructProperty->Struct, Outer, LooseProperties, Path);
 #if WITH_EDITORONLY_DATA
-			AsStructProperty->SetMetaData("TPSOverrideStructName", TPSOverrideStructName.ToString());
-			AsStructProperty->SetMetaData("OriginalType", TPSOverrideStructName.ToString());
+			AsStructProperty->SetMetaData("OriginalType", StructOriginalName.ToString());
 #endif
 		}
 		else if (const FArrayProperty* AsArrayProperty = CastField<FArrayProperty>(Property))
@@ -488,6 +483,20 @@ namespace UE
 		FWildcardPropertyPathName ParentPath;
 		return CreatePropertyBagArchetypeStructRec<UClass>(TemplateStruct, Outer, LooseProperties, ParentPath);
 	}
+
+	static void MarkPropertySetBySerialization(const UStruct* Struct, const void* StructData, const void* PropertyDataPtr)
+	{
+		if (const FSetProperty* ValuesSetByPropertyBagProperty = CastField<FSetProperty>(Struct->FindPropertyByName(Name_ValuesSetBySerialization)))
+		{
+			FScriptSetHelper ValuesSetByPropertyBag(ValuesSetByPropertyBagProperty, ValuesSetByPropertyBagProperty->ContainerPtrToValuePtr<void>(StructData));
+			const int64 ValueOffset = static_cast<const uint8*>(PropertyDataPtr) - static_cast<const uint8*>(StructData);
+			const int32 FoundIndex = ValuesSetByPropertyBag.FindElementIndex(&ValueOffset);
+			if (FoundIndex == INDEX_NONE)
+			{
+				ValuesSetByPropertyBag.AddElement(&ValueOffset);
+			}
+		}
+	}
 	
 	void MarkPropertySetBySerialization(UObject* Object, const FPropertyPathName& Path)
 	{
@@ -514,41 +523,21 @@ namespace UE
             {
                 ParentData = Params.Data;
                 ParentStruct = ParentAsStructProperty->Struct;
+            }
+            else
+            {
+            	// only mark properties set if they're in structs/classes
+	            return;
             }
 		}
-
 		// resolve the last segment of the path
 		Params.EndPathIndex += 1;
-        if (ensure(ResolvePropertyPathNameHelper(ParentStruct, Params)))
+		if (ensure(ResolvePropertyPathNameHelper(ParentStruct, Params)))
         {
-            MarkPropertySetBySerialization(ParentStruct, ParentData, Params.ResultProperty, Segment.Index);
+			MarkPropertySetBySerialization(ParentStruct, ParentData, Params.Data);
         }
 	}
 
-	void MarkPropertySetBySerialization(const UStruct* Struct, const void* StructData, const FProperty* Property, int32 ArrayIndex)
-	{
-		if (const FSetProperty* ValuesSetByPropertyBagProperty = CastField<FSetProperty>(Struct->FindPropertyByName(Name_ValuesSetBySerialization)))
-		{
-			const uint8* PropertyDataPtr;
-			if (ArrayIndex == INDEX_NONE || Property->IsA<FArrayProperty>() || Property->IsA<FMapProperty>() || Property->IsA<FSetProperty>())
-			{
-				PropertyDataPtr = Property->ContainerPtrToValuePtr<uint8>(StructData);
-			}
-			else
-			{
-				PropertyDataPtr = Property->ContainerPtrToValuePtr<uint8>(StructData, ArrayIndex);
-			}
-			
-			FScriptSetHelper ValuesSetByPropertyBag(ValuesSetByPropertyBagProperty, ValuesSetByPropertyBagProperty->ContainerPtrToValuePtr<void>(StructData));
-			const int64 ValueOffset = PropertyDataPtr - static_cast<const uint8*>(StructData);
-			const int32 FoundIndex = ValuesSetByPropertyBag.FindElementIndex(&ValueOffset);
-			if (FoundIndex == INDEX_NONE)
-			{
-				ValuesSetByPropertyBag.AddElement(&ValueOffset);
-			}
-		}
-	}
-
 	bool WasPropertySetBySerialization(UObject* Object, const FPropertyPathName& Path)
 	{
 		const FPropertyPathNameSegment Segment = Path.GetSegment(Path.GetSegmentCount() - 1);
diff --git a/Engine/Source/Runtime/CoreUObject/Private/UObject/ArchetypeUtils.h b/Engine/Source/Runtime/CoreUObject/Private/UObject/ArchetypeUtils.h
new file mode 100644
--- /dev/null
+++ b/Engine/Source/Runtime/CoreUObject/Private/UObject/ArchetypeUtils.h
@@ -0,0 +1,26 @@
+// Copyright Epic Games, Inc. All Rights Reserved.
+
+#pragma once
+#include "HAL/Platform.h"
+
+class UClass;
+class UObject;
+class UStruct;
+class FProperty;
+
+namespace UE
+{
+	class FPropertyBag;
+	class FPropertyPathName;
+
+	// generate a UClass that unions the properties of PropertyBag and TemplateStruct
+	UClass* CreatePropertyBagArchetypeClass(const FPropertyBag* PropertyBag, UStruct* TemplateStruct, UObject* Outer);
+
+	// notify that a property in an object was set when the object was deserialized
+	void MarkPropertySetBySerialization(UObject* Object, const FPropertyPathName& Path);
+	
+	// query whether a property in an object was set when the object was deserialized
+	bool WasPropertySetBySerialization(UObject* Object, const FPropertyPathName& Path);
+	// query whether a property in Struct was set when the struct was deserialized
+	bool WasPropertySetBySerialization(const UStruct* Struct, const void* StructData, const FProperty* Property, int32 ArrayIndex = 0);
+}
diff --git a/Engine/Source/Runtime/CoreUObject/Private/UObject/Class.cpp b/Engine/Source/Runtime/CoreUObject/Private/UObject/Class.cpp
--- a/Engine/Source/Runtime/CoreUObject/Private/UObject/Class.cpp
+++ b/Engine/Source/Runtime/CoreUObject/Private/UObject/Class.cpp
@@ -1569,6 +1569,11 @@ void UStruct::SerializeVersionedTaggedProperties(FStructuredArchive::FSlot Slot,
 					Property = CustomFindProperty(Tag.Name);
 				}
 
+				if (Property)
+				{
+					Property->AssignToTag(Tag);
+				}
+
 				if (SerializeContext && SerializeContext->bTrackSerializedPropertyPath)
 				{
 					const FName Name = Property ? Property->GetFName() : Tag.Name;
@@ -1595,8 +1600,6 @@ void UStruct::SerializeVersionedTaggedProperties(FStructuredArchive::FSlot Slot,
 
 				if (Property)
 				{
-					Property->AssignToTag(Tag);
-
 					FName PropID = Property->GetID();
 
 					// Check if this is a struct property and we have a redirector
@@ -1715,6 +1718,8 @@ void UStruct::SerializeVersionedTaggedProperties(FStructuredArchive::FSlot Slot,
 
 				if (SerializeContext && SerializeContext->bTrackSerializedPropertyPath)
 				{
+					// broadcast that a property was serialized
+					SerializeContext->OnTaggedPropertySerialize.Broadcast(*SerializeContext);
 					SerializeContext->SerializedPropertyPath.Pop();
 				}
 
diff --git a/Engine/Source/Runtime/CoreUObject/Private/UObject/PropertyArray.cpp b/Engine/Source/Runtime/CoreUObject/Private/UObject/PropertyArray.cpp
--- a/Engine/Source/Runtime/CoreUObject/Private/UObject/PropertyArray.cpp
+++ b/Engine/Source/Runtime/CoreUObject/Private/UObject/PropertyArray.cpp
@@ -662,6 +662,9 @@ void FArrayProperty::SerializeItem(FStructuredArchive::FSlot Slot, void* Value,
 				if (Context && Context->bTrackSerializedPropertyPath)
 				{
 					Context->SerializedPropertyPath.SetIndex(i);
+					
+					// broadcast that a property will be serialized
+					Context->OnTaggedPropertySerialize.Broadcast(*Context);
 				}
 				SerializeContainerItem(Array.EnterElement(), ArrayHelper.GetRawPtr(i));
 				PropertyNode = PropertyNode->PropertyListNext;
@@ -690,6 +693,9 @@ void FArrayProperty::SerializeItem(FStructuredArchive::FSlot Slot, void* Value,
 			if (Context && Context->bTrackSerializedPropertyPath)
 			{
 				Context->SerializedPropertyPath.SetIndex(i);
+					
+				// broadcast that a property will be serialized
+				Context->OnTaggedPropertySerialize.Broadcast(*Context);
 			}
 			SerializeContainerItem(Array.EnterElement(), ArrayHelper.GetRawPtr(i++));
 		}
@@ -1249,7 +1255,7 @@ bool FArrayProperty::LoadFromTag(const FPropertyTag& Tag)
 		// Skip property types that are missing the name of the inner type.
 		// Structs have their name in a tag in the serialized data, but we cannot
 		// proceed safely unless we know if the struct used native serialization.
-		if (!Property->IsA<FStructProperty>() &&
+		if (/*!Property->IsA<FStructProperty>() &&*/
 			!Property->IsA<FByteProperty>() &&
 			!Property->IsA<FEnumProperty>() &&
 			Property->LoadFromTag(InnerTag))
diff --git a/Engine/Source/Runtime/CoreUObject/Private/UObject/PropertyBagRepository.cpp b/Engine/Source/Runtime/CoreUObject/Private/UObject/PropertyBagRepository.cpp
--- a/Engine/Source/Runtime/CoreUObject/Private/UObject/PropertyBagRepository.cpp
+++ b/Engine/Source/Runtime/CoreUObject/Private/UObject/PropertyBagRepository.cpp
@@ -5,6 +5,11 @@
 #include "UObject/Object.h"
 #include "UObject/UObjectGlobals.h"
 #include "UObject/PropertyBag.h"
+#include "UObject/UObjectHash.h"
+#include "UObject/UObjectThreadContext.h"
+#include "UObject/LinkerLoad.h"
+#include "UObject/ArchetypeUtils.h"
+#include "UObject/Package.h"
 
 
 DEFINE_LOG_CATEGORY_STATIC(LogPropertyBagRepository, Log, All);
@@ -96,6 +101,17 @@ FPropertyBag* FPropertyBagRepository::CreateOuterBag(const UObjectBase* Owner)
 	return BagData->Bag;
 }
 
+UObject* FPropertyBagRepository::CreateArchetype(const UObjectBase* Owner)
+{
+	FPropertyBagRepositoryLock LockRepo(this);
+	FPropertyBagAssociationData& BagData = AssociatedData.FindOrAdd(Owner);
+	if(!BagData.Archetype)
+	{
+		CreateArchetypeUnsafe(Owner, BagData);
+	}
+	return BagData.Archetype;
+}
+
 // TODO: Remove this? Bag destruction to be handled entirely via UObject::BeginDestroy() (+ FPropertyBagProperty destructor)?
 void FPropertyBagRepository::DestroyOuterBag(const UObjectBase* Owner)
 {
@@ -150,11 +166,7 @@ bool FPropertyBagRepository::HasArchetype(const UObjectBase* Object) const
 UObject* FPropertyBagRepository::FindArchetype(const UObjectBase* Object)
 {
 	FPropertyBagRepositoryLock LockRepo(this);
-	FPropertyBagAssociationData* BagData = AssociatedData.Find(Object);
-	if(BagData && !BagData->Archetype)
-	{
-		CreateArchetypeUnsafe(Object, *BagData);
-	}
+	const FPropertyBagAssociationData* BagData = AssociatedData.Find(Object);
 	return BagData ? BagData->Archetype : nullptr;
 }
 
@@ -163,13 +175,50 @@ const UObject* FPropertyBagRepository::FindArchetype(const UObjectBase* Object)
 	return const_cast<FPropertyBagRepository*>(this)->FindArchetype(Object);
 }
 
-void FPropertyBagRepository::CreateArchetypeUnsafe(const UObjectBase* Object, FPropertyBagAssociationData& BagData)
+bool FPropertyBagRepository::WasPropertySetBySerialization(UObject* Object, const FPropertyPathName& Path)
+{
+	return UE::WasPropertySetBySerialization(Object, Path);
+}
+
+bool FPropertyBagRepository::WasPropertySetBySerialization(const UStruct* Struct, const void* StructData, const FProperty* Property, int32 ArrayIndex)
 {
-	check(BagData.Archetype);	// No repeated calls
+	return UE::WasPropertySetBySerialization(Struct, StructData, Property, ArrayIndex);
+}
+
+void FPropertyBagRepository::CreateArchetypeUnsafe(const UObjectBase* Owner, FPropertyBagAssociationData& BagData)
+{
+	check(!BagData.Archetype);	// No repeated calls
+	const FPropertyBag* PropertyBag = BagData.Bag;
+	// construct archetype class
+	// TODO: should we put the archetype or it's class in a package?
+	const UClass* ArchetypeClass = CreatePropertyBagArchetypeClass(PropertyBag, Owner->GetClass(), GetTransientPackage());
+
+	// construct archetype object
+	FStaticConstructObjectParameters Params(ArchetypeClass);
+	Params.SetFlags |= EObjectFlags::RF_Transactional;
+	UObject* ArchetypeObject = StaticConstructObject_Internal(Params);
+	BagData.Archetype = ArchetypeObject;
+	
+	// setup load context to mark properties the that were set by serialization
+	FUObjectSerializeContext* LoadContext = FUObjectThreadContext::Get().GetSerializeContext();
+	TGuardValue<bool> ScopedImpersonateProperties(LoadContext->bImpersonateProperties, true);
+	TGuardValue<bool> ScopedTrackSerializedPropertyPath(LoadContext->bTrackSerializedPropertyPath, true);
 	
-	// TODO: Implementation
-	// TODO: AddReferencedObject/AddToRoot GC handling.
-	BagData.Archetype = nullptr;	// Guaranteed nullptr on failure.
+	const FDelegateHandle OnTaggedPropertySerializeHandle = LoadContext->OnTaggedPropertySerialize.AddLambda(
+		[&BagData](const FUObjectSerializeContext& Context)
+		{
+			if (!Context.SerializedPropertyPath.IsEmpty())
+			{
+				MarkPropertySetBySerialization(BagData.Archetype, Context.SerializedPropertyPath);
+			}
+		}
+	);
+
+	UObject* OwnerAsObject = (UObject*)Owner;
+	OwnerAsObject->SetFlags(RF_NeedLoad);
+	OwnerAsObject->GetLinker()->Preload(OwnerAsObject);
+
+	LoadContext->OnTaggedPropertySerialize.Remove(OnTaggedPropertySerializeHandle);
 }
 
 // Not sure this is necessary.
diff --git a/Engine/Source/Runtime/CoreUObject/Private/UObject/PropertyMap.cpp b/Engine/Source/Runtime/CoreUObject/Private/UObject/PropertyMap.cpp
--- a/Engine/Source/Runtime/CoreUObject/Private/UObject/PropertyMap.cpp
+++ b/Engine/Source/Runtime/CoreUObject/Private/UObject/PropertyMap.cpp
@@ -731,6 +731,9 @@ void FMapProperty::SerializeItem(FStructuredArchive::FSlot Slot, void* Value, co
 				if (Context && Context->bTrackSerializedPropertyPath)
 				{
 					Context->SerializedPropertyPath.SetIndex(Index);
+					
+					// broadcast that a property will be serialized
+					Context->OnTaggedPropertySerialize.Broadcast(*Context);
 				}
 
 				// TODO: Need a way to indicate that we are serializing a key and value into the property bag. Push Key/Value names?
diff --git a/Engine/Source/Runtime/CoreUObject/Private/UObject/PropertySet.cpp b/Engine/Source/Runtime/CoreUObject/Private/UObject/PropertySet.cpp
--- a/Engine/Source/Runtime/CoreUObject/Private/UObject/PropertySet.cpp
+++ b/Engine/Source/Runtime/CoreUObject/Private/UObject/PropertySet.cpp
@@ -307,6 +307,9 @@ void FSetProperty::SerializeItem(FStructuredArchive::FSlot Slot, void* Value, co
 				if (Context && Context->bTrackSerializedPropertyPath)
 				{
 					Context->SerializedPropertyPath.SetIndex(Index);
+					
+					// broadcast that a property will be serialized
+					Context->OnTaggedPropertySerialize.Broadcast(*Context);
 				}
 				ElementProp->SerializeItem(ElementsArray.EnterElement(), SetHelper.GetElementPtrWithoutCheck(Index));
 			}
diff --git a/Engine/Source/Runtime/CoreUObject/Public/UObject/ArchetypeUtils.h b/Engine/Source/Runtime/CoreUObject/Public/UObject/ArchetypeUtils.h
deleted file mode 100644
diff --git a/Engine/Source/Runtime/CoreUObject/Public/UObject/UObjectThreadContext.h b/Engine/Source/Runtime/CoreUObject/Public/UObject/UObjectThreadContext.h
--- a/Engine/Source/Runtime/CoreUObject/Public/UObject/UObjectThreadContext.h
+++ b/Engine/Source/Runtime/CoreUObject/Public/UObject/UObjectThreadContext.h
@@ -177,6 +177,9 @@ public:
 	/** True when the SerializedObject properties are being impersonated. */
 	UE_INTERNAL bool bImpersonateProperties;
 
+	/** event called after each tagged property is deserialized */
+	DECLARE_MULTICAST_DELEGATE_OneParam(FOnTaggedPropertySerialized, const FUObjectSerializeContext&)
+	UE_INTERNAL FOnTaggedPropertySerialized OnTaggedPropertySerialize;
 
 	/** Adds a new loaded object */
 	COREUOBJECT_API void AddLoadedObject(UObject* InObject);
