[Feature] Made it so objects will display as archetypes in the details panel if an archetype exists
[BugFix] Inlined Subobjects not displaying correctly in left panel of archetype fixup tool

diff --git a/Engine/Plugins/Editor/ArchetypeFixupTool/Source/Private/ArchetypeFixupDetailCustomization.cpp b/Engine/Plugins/Editor/ArchetypeFixupTool/Source/Private/ArchetypeFixupDetailCustomization.cpp
--- a/Engine/Plugins/Editor/ArchetypeFixupTool/Source/Private/ArchetypeFixupDetailCustomization.cpp
+++ b/Engine/Plugins/Editor/ArchetypeFixupTool/Source/Private/ArchetypeFixupDetailCustomization.cpp
@@ -10,6 +10,7 @@
 #include "Widgets/Layout/SWidgetSwitcher.h"
 #include "Widgets/Input/SComboButton.h"
 #include "Framework/MultiBox/MultiBoxBuilder.h"
+#include "UObject/PropertyBagRepository.h"
 
 #define LOCTEXT_NAMESPACE "ArchetypeFixupDetails"
 
@@ -58,6 +59,19 @@ void FArchetypeFixupDetailNodeBuilder::GetRedirectOptions(const FProperty* Prope
 	{
 		GetRedirectOptions(AsStructProperty->Struct, Value, Path, OutPaths);
 	}
+	if (const FObjectProperty* AsObjectProperty = CastField<FObjectProperty>(Property))
+	{
+		if (AsObjectProperty->HasAnyPropertyFlags(CPF_InstancedReference))
+		{
+			TObjectPtr<UObject> Object = AsObjectProperty->GetObjectPropertyValue(Value);
+			UE::FPropertyBagRepository& PropertyBagRepository = UE::FPropertyBagRepository::Get();
+			if (UObject* Found = PropertyBagRepository.FindArchetype(Object))
+			{
+				Object = Found;
+			}
+			GetRedirectOptions(Object->GetClass(), Object, Path, OutPaths);
+		}
+	}
 	else if (const FArrayProperty* AsArrayProperty = CastField<FArrayProperty>(Property))
 	{
 		FScriptArrayHelper Array(AsArrayProperty, Value);
@@ -208,6 +222,21 @@ void FArchetypeFixupDetailNodeBuilder::GenerateHeaderRowContent(FDetailWidgetRow
 	{
 		return;
 	}
+
+	// handle categories of sub-objects
+	if (PropertyHandle->IsCategoryHandle())
+	{
+		NodeRow
+		[
+			SNew(SBox)
+			.VAlign(VAlign_Center)
+			[
+				SNew(STextBlock)
+				.Text(PropertyHandle->GetPropertyDisplayName())
+			]
+		];
+		return;
+	}
 	
 	TSharedPtr<SWidget> InnerNameContent = PropertyHandle->CreatePropertyNameWidget();
 	if (PropertyHandle->GetKeyHandle())
@@ -280,6 +309,22 @@ void FArchetypeFixupDetailNodeBuilder::GenerateChildContent(IDetailChildrenBuild
 	
 	uint32 ChildCount;
 	PropertyHandle->GetNumChildren(ChildCount);
+	if (ChildCount == 1 && PropertyHandle->GetProperty() && PropertyHandle->GetProperty()->IsA<FObjectProperty>())
+	{
+		const TSharedRef<IPropertyHandle> ChildHandle = PropertyHandle->GetChildHandle(0).ToSharedRef();
+		if (!ChildHandle->GetProperty())
+		{
+			uint32 GrandChildCount;
+			ChildHandle->GetNumChildren(GrandChildCount);
+			for (uint32 I = 0; I < GrandChildCount; ++I)
+			{
+				const TSharedRef<IPropertyHandle> GrandChildHandle = ChildHandle->GetChildHandle(I).ToSharedRef();
+				ChildrenBuilder.AddCustomBuilder(MakeShared<FArchetypeFixupDetailNodeBuilder>(Panel.ToSharedRef(), GrandChildHandle));
+			}
+			return;
+		}
+	}
+	
 	for (uint32 I = 0; I < ChildCount; ++I)
 	{
 		const TSharedRef<IPropertyHandle> ChildHandle = PropertyHandle->GetChildHandle(I).ToSharedRef();
diff --git a/Engine/Plugins/Editor/ArchetypeFixupTool/Source/Private/ArchetypeFixupPanel.cpp b/Engine/Plugins/Editor/ArchetypeFixupTool/Source/Private/ArchetypeFixupPanel.cpp
--- a/Engine/Plugins/Editor/ArchetypeFixupTool/Source/Private/ArchetypeFixupPanel.cpp
+++ b/Engine/Plugins/Editor/ArchetypeFixupTool/Source/Private/ArchetypeFixupPanel.cpp
@@ -7,9 +7,9 @@
 #include "DetailTreeNode.h"
 #include "Widgets/Layout/LinkableScrollBar.h"
 #include "ArchetypeFixupDetailCustomization.h"
-#include "UObject/ArchetypeUtils.h"
 #include "Modules/ModuleManager.h"
 #include "Editor.h"
+#include "UObject/PropertyBagRepository.h"
 
 #define LOCTEXT_NAMESPACE "ArchetypeFixupPanel"
 
@@ -291,7 +291,7 @@ bool FArchetypeFixupPanel::AreAllConflictsRedirected() const
 			{
 				if (const TSharedPtr<IPropertyHandle> Handle = TreeNode->CreatePropertyHandle())
 				{
-					if (!MarkedForDelete.Contains(*Handle->CreateFPropertyPath()))
+					if (!Handle->IsCategoryHandle() && !MarkedForDelete.Contains(*Handle->CreateFPropertyPath()))
 					{
 						bFoundConflict = true;
 						return ETreeTraverseControl::Break;
@@ -351,7 +351,13 @@ static void* ResolvePath(const FPropertyPath& Path, void* Value)
 
 		if (const FObjectProperty* AsObjectProperty = CastField<FObjectProperty>(Property))
 		{
-			Value = AsObjectProperty->GetObjectPropertyValue(Value);
+			UObject* Object = AsObjectProperty->GetObjectPropertyValue(Value);
+			UE::FPropertyBagRepository& PropertyBagRepository = UE::FPropertyBagRepository::Get();
+			if (UObject* Found = PropertyBagRepository.FindArchetype(Object))
+			{
+				Object = Found;
+			}
+			Value = Object;
 		}
 		else if (PathIndex + 1 < Path.GetNumProperties())
 		{
@@ -516,7 +522,7 @@ static void InitRedirectedPropertyTreeRec(const TSharedPtr<FRedirectedPropertyNo
 	{
 		if (Property->ArrayDim == 1)
 		{
-			if (UE::WasPropertySetBySerialization(Struct, StructValue, Property))
+			if (UE::FPropertyBagRepository::WasPropertySetBySerialization(Struct, StructValue, Property))
 			{
 				const TSharedPtr<FRedirectedPropertyNode>& ChildNode = Node->FindOrAdd(FPropertyInfo(Property));
 				void* Value = Property->ContainerPtrToValuePtr<void>(StructValue);
@@ -527,7 +533,7 @@ static void InitRedirectedPropertyTreeRec(const TSharedPtr<FRedirectedPropertyNo
 		{
 			for (int32 StaticArrayIndex = 0; StaticArrayIndex < Property->ArrayDim; ++StaticArrayIndex)
             {
-            	if (UE::WasPropertySetBySerialization(Struct, StructValue, Property, StaticArrayIndex))
+            	if (UE::FPropertyBagRepository::WasPropertySetBySerialization(Struct, StructValue, Property, StaticArrayIndex))
             	{
             		const TSharedPtr<FRedirectedPropertyNode>& ChildNode = Node->FindOrAdd(FPropertyInfo(Property, StaticArrayIndex));
             		void* Value = Property->ContainerPtrToValuePtr<void>(StructValue, StaticArrayIndex);
@@ -544,6 +550,22 @@ static void InitRedirectedPropertyTreeRec(const TSharedPtr<FRedirectedPropertyNo
 	{
 		InitRedirectedPropertyTreeRec(Node, AsStructProperty->Struct, Value);
 	}
+	else if (const FObjectProperty* AsObjectProperty = CastField<FObjectProperty>(Property))
+	{
+		if (AsObjectProperty->HasAnyPropertyFlags(CPF_InstancedReference))
+		{
+			if (UObject* Object = AsObjectProperty->GetObjectPropertyValue(Value))
+            {
+				UE::FPropertyBagRepository& PropertyBagRepository = UE::FPropertyBagRepository::Get();
+				if (UObject* Found = PropertyBagRepository.FindArchetype(Object))
+				{
+					Object = Found;
+				}
+            	InitRedirectedPropertyTreeRec(Node, Object->GetClass(), Object);
+            }
+		}
+		
+	}
 	else if (const FArrayProperty* AsArrayProperty = CastField<FArrayProperty>(Property))
 	{
 		FScriptArrayHelper Array(AsArrayProperty, Value);
diff --git a/Engine/Plugins/Editor/ArchetypeFixupTool/Source/Private/ArchetypeFixupTool.cpp b/Engine/Plugins/Editor/ArchetypeFixupTool/Source/Private/ArchetypeFixupTool.cpp
--- a/Engine/Plugins/Editor/ArchetypeFixupTool/Source/Private/ArchetypeFixupTool.cpp
+++ b/Engine/Plugins/Editor/ArchetypeFixupTool/Source/Private/ArchetypeFixupTool.cpp
@@ -37,10 +37,16 @@ public:
 		const TSharedPtr<IPropertyHandle> PropertyHandleB = PinnedTreeNodeB->CreatePropertyHandle();
 		if (!PropertyHandleA || !PropertyHandleB)
 		{
-			// category nodes
+			// top level category nodes
 			return PinnedTreeNodeA->GetNodeName() == PinnedTreeNodeB->GetNodeName();
 		}
-
+		
+		if (PropertyHandleA->IsCategoryHandle() || PropertyHandleB->IsCategoryHandle())
+		{
+			// category nodes
+			return PropertyHandleA->GetPropertyDisplayName().ToString() == PropertyHandleB->GetPropertyDisplayName().ToString();
+		}
+		
 		return PropertyHandleA->GetProperty() == PropertyHandleB->GetProperty();
 	}
 
diff --git a/Engine/Source/Editor/PropertyEditor/Private/ObjectPropertyNode.cpp b/Engine/Source/Editor/PropertyEditor/Private/ObjectPropertyNode.cpp
--- a/Engine/Source/Editor/PropertyEditor/Private/ObjectPropertyNode.cpp
+++ b/Engine/Source/Editor/PropertyEditor/Private/ObjectPropertyNode.cpp
@@ -8,6 +8,7 @@
 #include "ObjectEditorUtils.h"
 #include "EditorCategoryUtils.h"
 #include "PropertyEditorHelpers.h"
+#include "UObject/PropertyBagRepository.h"
 
 #if WITH_EDITORONLY_DATA
 #include "Engine/Blueprint.h"
@@ -54,13 +55,21 @@ const UPackage* FObjectPropertyNode::GetUPackage(int32 InIndex) const
 // Adds a new object to the list.
 void FObjectPropertyNode::AddObject( UObject* InObject )
 {
+	UE::FPropertyBagRepository& Repository = UE::FPropertyBagRepository::Get();
+	if (UObject* Found = Repository.FindArchetype(InObject))
+	{
+		InObject = Found;
+	}
 	Objects.Add( InObject );
 }
 
 // Adds new objects to the list.
 void FObjectPropertyNode::AddObjects(const TArray<UObject*>& InObjects)
 {
-	Objects.Append( InObjects );
+	for (UObject* Object : InObjects)
+	{
+		AddObject(Object);
+	}
 }
 
 // Removes an object from the list.
@@ -72,6 +81,14 @@ void FObjectPropertyNode::RemoveObject( UObject* InObject )
 	{
 		Objects.RemoveAt( idx, 1 );
 	}
+	else
+	{
+		UE::FPropertyBagRepository& Repository = UE::FPropertyBagRepository::Get();
+		if (UObject* Found = Repository.FindArchetype(InObject))
+		{
+			RemoveObject(Found);
+		}
+	}
 }
 
 // Removes all objects from the list.
diff --git a/Engine/Source/Editor/PropertyEditor/Private/PropertyHandleImpl.cpp b/Engine/Source/Editor/PropertyEditor/Private/PropertyHandleImpl.cpp
--- a/Engine/Source/Editor/PropertyEditor/Private/PropertyHandleImpl.cpp
+++ b/Engine/Source/Editor/PropertyEditor/Private/PropertyHandleImpl.cpp
@@ -2313,6 +2313,11 @@ bool FPropertyHandleBase::IsFavorite() const
 	return false;
 }
 
+bool FPropertyHandleBase::IsCategoryHandle() const
+{
+	return Implementation->GetPropertyNode().IsValid() && Implementation->GetPropertyNode()->AsCategoryNode();
+}
+
 bool FPropertyHandleBase::IsCustomized() const
 {
 	return Implementation->GetPropertyNode()->HasNodeFlags( EPropertyNodeFlags::IsCustomized ) != 0;
diff --git a/Engine/Source/Editor/PropertyEditor/Private/PropertyHandleImpl.h b/Engine/Source/Editor/PropertyEditor/Private/PropertyHandleImpl.h
--- a/Engine/Source/Editor/PropertyEditor/Private/PropertyHandleImpl.h
+++ b/Engine/Source/Editor/PropertyEditor/Private/PropertyHandleImpl.h
@@ -517,6 +517,7 @@ public:
 	virtual int32 GetArrayIndex() const override;
 	virtual void RequestRebuildChildren() override;
 	virtual bool IsFavorite() const override;
+	virtual bool IsCategoryHandle() const override;
 
 	PROPERTYEDITOR_API TSharedPtr<FPropertyNode> GetPropertyNode() const;
 	void OnCustomResetToDefault(const FResetToDefaultOverride& OnCustomResetToDefault);
diff --git a/Engine/Source/Editor/PropertyEditor/Private/PropertyNode.cpp b/Engine/Source/Editor/PropertyEditor/Private/PropertyNode.cpp
--- a/Engine/Source/Editor/PropertyEditor/Private/PropertyNode.cpp
+++ b/Engine/Source/Editor/PropertyEditor/Private/PropertyNode.cpp
@@ -22,6 +22,7 @@
 #include "UObject/MetaData.h"
 #include "UObject/TextProperty.h"
 #include "UObject/EnumProperty.h"
+#include "UObject/PropertyBagRepository.h"
 #include "UObject/UnrealType.h"
 
 #include "UObject/PropertyOptional.h"
@@ -692,6 +693,11 @@ EPropertyDataValidationResult FPropertyNode::EnsureDataIsValid()
 				if (ObjectProperty && !bIgnoreAllMismatch)
 				{
 					UObject* Obj = ObjectProperty->GetObjectPropertyValue(Addr);
+					UE::FPropertyBagRepository& Repository = UE::FPropertyBagRepository::Get();
+					if (UObject* Found = Repository.FindArchetype(Obj))
+					{
+						Obj = Found;
+					}
 					if (IsValid(Obj))
 					{
 						if (!bShowInnerObjectPropertiesObjectChanged && 
diff --git a/Engine/Source/Editor/PropertyEditor/Public/PropertyHandle.h b/Engine/Source/Editor/PropertyEditor/Public/PropertyHandle.h
--- a/Engine/Source/Editor/PropertyEditor/Public/PropertyHandle.h
+++ b/Engine/Source/Editor/PropertyEditor/Public/PropertyHandle.h
@@ -746,6 +746,8 @@ public:
 	* It does not handle the property being moved to another category during customization
 	*/
 	virtual FText GetDefaultCategoryText() const = 0;
+
+	virtual bool IsCategoryHandle() const = 0;
 };
 
 /**
